ushell - A command line interface for the Unreal Engine

#rnx

diff --git a/Engine/Extras/ushell/README.txt b/Engine/Extras/ushell/README.txt
new file mode 100644
--- /dev/null
+++ b/Engine/Extras/ushell/README.txt
@@ -0,0 +1,90 @@
+
+ ushell - A command line interface for the Unreal Engine
+
+# Quick Start Guide
+
+1. Create a shortcut to `ushell.exe`
+2. Set the start up directory to the location of a .uproject
+3. Start ushell.
+4. Run `.help` to see available commands.
+
+The properties of the shortcut should look like this;
+
+    - Target:   m:\branch\Extras\ushell\ushell.exe
+    - Start in: m:\branch\MyBranch
+
+Alternatively you can use the '--project=' command line argument;
+
+    - Target:   m:\branch\Extras\ushell\ushell.exe --project=d:\branch\MyProject
+    - Start in: m:\my\favourite\directory
+
+The `.project [path_to_uproject]` command can be used to change the session's
+active project from within ushell.
+
+Use `.help readme` for more detailed information. This includes details on how 
+to use alternative terminals and shells such as PowerShell and how to extend and 
+customize ushell.
+
+# Commands
+
+Perhaps the best way to understand some commands available in ushell and what
+they can do is to list a few examples;
+
+ 1. .build editor
+ 2. .build game win64
+ 3. .build program UnrealInsights shipping
+ 4. .cook game win64
+ 5. .stage game win64
+ 6. .run editor
+ 7. .run game win64 --trace -- -ExecCmds="Zippy Bungle Rainbow"
+ 8. .run program UnrealInsights shipping
+ 9. .p4 cherrypick 1234567
+10. .sln generate
+11. .sln open
+12. .info
+
+Each of ushell's commands accept a '--help' argument which will show
+documentation for the command, details on how to its invoked, and descriptions
+for the available options.
+
+# Tab Completion And Command History
+
+Anyone familiar with editing commands and recalling previous ones in Bash (i.e.
+Readline) will feel at home at a ushell prompt.
+
+There is extensive context-sensitive Tab completion available for commands and
+their arguments. Hitting Tab can help both to discover commands and their
+arguments, and also aid in fast convenient command entry. For example;
+
+1. .<tab><tab>            : displays available commands
+2. .b<tab>                : completes ".build "
+3. .run <tab><tab>        : shows options for the .run commands first argument
+4. .build editor --p<tab> : adds "--platform=" (further Tabs complete platforms)
+
+Ushell maintains a history of commands run which is persisted from one session to
+the next. Previous commands can be conveniently recalled in a few ways. To step
+backwards through prior commands by prefix use PgUp;
+
+1. .bu<pgup>                : Cycle back through commands that started with ".bu"
+2. .run game switch<pgup>   : Iterate through previous runs on Switch.
+
+A more thorough incremental history search is done with Ctrl-R. This displays a
+prompt to enter a search string and will display the latest command with a match.
+Further Ctrl-R hits will step backwards through commands that match the search
+string (with Ctrl-S stepping forwards). History searching is case-sensitive.
+
+# What does ushell.exe do?
+
+ushell is built around a Python-based framework. The code and supporting files 
+are embedded inside the `ushell.exe`. By default, `ushell.exe` extracts these 
+files into `$LocalAppData/ushell` and launches the newest version. When ushell
+first starts it downloads an appropriate version of Python and other dependent 
+utilities - there are no prerequisites the user need worry about.
+
+To list the files extract, or dump the `ushell.zip` file for more advanced 
+deployment scenarios, the arguments `ushell.exe list` and `ushell.exe dump` can
+be used respectively.
+
+
+
+vim: tw=80 fo=wnt ft=markdown nosi spell
