Updated ushell [git:56a96bd6]

diff --git a/Engine/Extras/ushell/README.txt b/Engine/Extras/ushell/README.txt
--- a/Engine/Extras/ushell/README.txt
+++ b/Engine/Extras/ushell/README.txt
@@ -4,8 +4,8 @@
 # Quick Start Guide
 
 1. Create a shortcut to `ushell.exe`
-2. Set the start up directory to the location of a .uproject
-3. Start ushell.
+2. Set the start up directory to the location of a .uproject file
+3. Start ushell via the shortcut
 4. Run `.help` to see available commands.
 
 The properties of the shortcut should look like this;
