Updated ushell [26875160]

diff --git a/Engine/Extras/ushell/.gitattributes b/Engine/Extras/ushell/.gitattributes
new file mode 100644
--- /dev/null
+++ b/Engine/Extras/ushell/.gitattributes
@@ -0,0 +1 @@
+*.sh eol=lf
diff --git a/Engine/Extras/ushell/.gitignore b/Engine/Extras/ushell/.gitignore
new file mode 100644
--- /dev/null
+++ b/Engine/Extras/ushell/.gitignore
@@ -0,0 +1,5 @@
+.*/
+*.pyc
+tags
+_build/
+__pycache__/
diff --git a/Engine/Extras/ushell/README.txt b/Engine/Extras/ushell/README.txt
--- a/Engine/Extras/ushell/README.txt
+++ b/Engine/Extras/ushell/README.txt
@@ -3,27 +3,34 @@
 
 # Quick Start Guide
 
-1. Create a shortcut to `ushell.exe`
+1. Create a shortcut to 'ushell.bat'
 2. Set the start up directory to the location of a .uproject
 3. Start ushell.
-4. Run `.help` to see available commands.
+4. Run '.help' to see available commands.
 
 The properties of the shortcut should look like this;
 
-    - Target:   m:\branch\Extras\ushell\ushell.exe
-    - Start in: m:\branch\MyBranch
+    - Target:   d:\ushell\ushell.bat
+    - Start in: d:\perforce\branch\QAGame
 
 Alternatively you can use the '--project=' command line argument;
 
-    - Target:   m:\branch\Extras\ushell\ushell.exe --project=d:\branch\MyProject
-    - Start in: m:\my\favourite\directory
+    - Target:   d:\ushell\ushell.bat --project=d:\perforce\branch\QAGame
+    - Start in: d:\my\favourite\directory
 
 The `.project [path_to_uproject]` command can be used to change the session's
 active project from within ushell.
 
-Use `.help readme` for more detailed information. This includes details on how 
-to use alternative terminals and shells such as PowerShell and how to extend and 
-customize ushell.
+# Mac and Linux
+
+On POSIX-based platforms ushell works by establishing itself in the current
+shell;
+
+  source ushell.sh
+
+The arguments and working directory rules in the Quick Start Guide above also
+apply when sourcing ushell.sh. Ensuring a suitable version of Python is available
+(and possibly a toolchain to build Pips) is left to the user to take care of.
 
 # Commands
 
@@ -31,17 +38,16 @@ Perhaps the best way to understand some commands available in ushell and what
 they can do is to list a few examples;
 
  1. .build editor
- 2. .build game win64
+ 2. .build game ps4
  3. .build program UnrealInsights shipping
- 4. .cook game win64
- 5. .stage game win64
+ 4. .cook game ps4
+ 5. .stage game ps4
  6. .run editor
- 7. .run game win64 --trace -- -ExecCmds="Zippy Bungle Rainbow"
+ 7. .run game ps4 --trace -- -ExecCmds="Zippy Bungle Rainbow"
  8. .run program UnrealInsights shipping
  9. .p4 cherrypick 1234567
 10. .sln generate
-11. .sln open
-12. .info
+11. .info
 
 Each of ushell's commands accept a '--help' argument which will show
 documentation for the command, details on how to its invoked, and descriptions
@@ -73,17 +79,206 @@ prompt to enter a search string and will display the latest command with a match
 Further Ctrl-R hits will step backwards through commands that match the search
 string (with Ctrl-S stepping forwards). History searching is case-sensitive.
 
-# What does ushell.exe do?
+# Scripting
+
+There is modest support for scripting ushell with Batch scripts. This can be
+useful for example to schedule overnight sync-builds. Here is a minimal example;
+
+```
+@echo off
+cd /d d:\branch\myuproj
+call d:\ushell\ushell.bat
+.p4 sync --all
+.build editor
+.build client ps4
+```
+
+# Chaining Commands
+
+It is also possible to chain commands entered on the command line with '&&';
+
+  .build editor && .cook game ps4
+
+The '&&' will only execute the subsequent command if the previous one succeeded.
+Using a single '&' unconditionally runs the next command.
+
+# Customising the Shell
+
+If the host shell is the standard Windows command prompt then ushell will check
+for and run `$USERPROFILE/.ushell/hooks/startup.bat` allowing the user to extend
+the session. Here is an example script that adds a simple `.bcr` alias;
+
+```
+@echo off
+doskey .bcr=.uat BuildCookRun -- $*
+```
+
+# Alternative Terminals
+
+To use ushell in an alternative terminal such as Windows Terminal or VSCode's
+integrated terminal ushell should be started as follows;
+
+  cmd.exe /d/k d:\ushell\ushell.bat --project=d:\branch\myuproj\myuproj.uproject
+
+This is required because by default ushell.bat detects if it was launched
+explicitly through Explorer (or a shortcut) and if not the assumption is it is
+running in a non-interactive scripting context (see `Scripting`)
+
+# Alternative Host Shells
+
+A lot of the above assumes that cmd.exe is the host shell. Other shells are also
+supported; Bash or Zsh on POSIX-based platforms, and PowerShell. There are
+scripts in ushell's root folder for using these alternative shells.
+
+## PowerShell
+
+PowerShell integration works by importing ushell as a module. Set your
+`PSModulePath` to contain the ushell's directory (this can be done in your
+`$PROFILE` file);
+
+```
+$env:PSModulePath = "$($env:PSModulePath);c:\path\to\ushell\"
+Import-Module ushell
+```
+
+If you use a PowerShell prompt enhancer like oh-my-posh you can extract
+environment variables from ushell to populate your prompt/window title/etc. The
+ushell-related environment variables can be list with `dir env:USHELL*`.
+
+```
+function Set-MyPoshContext {
+    if ($Global:WhatIfPreference) {
+        $env:POSH_WHATIF = "What If"
+    }
+    else {
+        $env:POSH_WHATIF = ""
+    }
+    if ( $null -ne (Get-Module ushell)) {
+        Update-UShellEnvVars | Out-Null
+    }
+}
+New-Alias -Name 'Set-PoshContext' -Value 'Set-MyPoshContext' -Scope Global -Force
+```
+
+# Extending ushell
+
+## Case Study
+
+A ushell instance is consists of "flow" which provides the framework for building
+a shell and commands for it, and channels, where a channel is a group of commands
+and tools. Flow enumerates channels by looking in the `channels/` folder for a
+well-known file that describes the channel. On disk a channel may look like;
+
+```
+channels/
+  mychannel/
+    describe.flow.py
+    cmds/
+      mycmd.py
+```
+
+The `describe.flow.py` is how ushell detects the subfolder as a channel. The file
+describes what commands the channel adds and how they are invoked. There is also
+an API for describing tools and how to fetch them, and a basic plugin mechanism.
+Here is a basic example channel with a single command;
+
+```
+import flow.describe
+
+# Describe the channel's commands. source() states where the source .py is and
+# the class that implements the command. The invoke() method sets how to the
+# command should be invoked from the shell - "d:\>.mycmd zippy" in this example
+my_cmd = flow.describe.Command()
+my_cmd.source("cmds/mycmd.py", "MyCmdClass")
+my_cmd.invoke("mycmd", "zippy")
+
+# Describe the channel. The pip() method can be used to install pip from the
+# PyPi repository. version() can be used to invalidate a channel when updates
+# are pulled. The parent() forms channels into a tree and informs the order of
+# inheritance when overriding commands.
+channel = flow.describe.Channel()
+channel.parent("unreal.core")
+channel.version("0")
+#channel.pip(...)
+```
+
+A command is implemented by deriving a class from `flow.cmd.Cmd` and using a
+docstring and `flow.cmd.Arg/Opt` classes to set a description and specify the
+commands arguments. Given the example channel above the `mycmd.py` might look
+as follows;
+
+```
+import flow.cmd
+
+class MyCmdClass(flow.cmd.Cmd):
+    """ This is the description of my command shown when --help is given """
+
+    # Describe the commands paramsters with Arg and Opt objects
+    argone = flow.cmd.Arg(str,   "A positional argument that must be given")
+    argtwo = flow.cmd.Arg("two", "A positional argument with a default value")
+    option = flow.cmd.Opt(False, "An boolean optional argument; --option")
+    param  = flow.cmd.Opt(5,     "An optional argument that takes a parameter; --param=6")
+
+    # Add completion for arguments with a complete_[argname] method. Completion
+    # methods return something iterable (thus they can be generator functions).
+    def complete_argone(self, prefix):
+        return ("zippy", "bungle", "george")
+
+    # The entry point of the command. Arguments are prepopulated in self.args
+    def main(self):
+        if self.args.option:
+            print(self.args.argument)
+```
+
+For working with Unreal branches there is an `unrealcmd` module that specialises
+a `flow.cmd.Cmd` object. The most useful method is `get_unreal_context()` which
+gives access to the branch, project, targets, and more (details can be found by
+reading `channels/unreal/core/pylib/unreal/_context.py`). An example;
+
+```
+import unrealcmd
+
+class MyUnrealCmdClass(unrealcmd.Cmd):
+    """ This command works with an Unreal branch and/or project """
+    argument = unrealcmd.Arg(int,   "...")
+    optional = unrealcmd.Opt(False, "...")
+
+    ue_context = self.get_unreal_context()
+    if project := ue_context.get_project():
+        print("Project name", project.get_name())
+```
+
+## Miscellaneous
+
+It possible to inject into existing commands by matching the words used to invoke
+them. So if the case study example was changed to `my_cmd.invoke("build", "editor")`
+then executing `.build editor` would in fact call `MyCmdClass.main()`. Calling
+`super().main()` would in turn run the proper ushell command to build the editor.
+
+There are a few more details about channels not covered thus far. Briefly; a
+channel's `pylib` folder is added to PYTHON_PATH so channels can define modules
+(handy for sharing between commands). Channels can participate in the boot
+process and prompt building - examples can be found in channels with boot and
+prompt commands.
+
+Channel names are derived from their file system location; `channels/geoffrey`
+gets the name `geoffrey`, and `channels/mr/hayes` becomes `mr.hayes".
+
+## Setting Yourself up to Extend ushell
+
+Currently the best way to add channels is to sync create a clientspec with the
+depot path `//depot/usr/martin.ridgers/ushell` in view along with your own
+channel mapped into the `channel/X` folder.
+
+Extending the UGS-deployed ushell is in essence impossible at this time for two
+reasons; 1) the deployment is overwritten by UGS when a new version becomes
+available, and 2) flow only looks in the channels/ sub-folder of where it is
+deployed. _This is going to change_ so that users can add channels in a
+well-known location that any distribution of ushell can find. To. Do.
 
-ushell is built around a Python-based framework. The code and supporting files 
-are embedded inside the `ushell.exe`. By default, `ushell.exe` extracts these 
-files into `$LocalAppData/ushell` and launches the newest version. When ushell
-first starts it downloads an appropriate version of Python and other dependent 
-utilities - there are no prerequisites the user need worry about.
+# Contact
 
-To list the files extract, or dump the `ushell.zip` file for more advanced 
-deployment scenarios, the arguments `ushell.exe list` and `ushell.exe dump` can
-be used respectively.
[REDACTED_EMAIL]
 
 
 
diff --git a/Engine/Extras/ushell/channels/flow/core/boot.py b/Engine/Extras/ushell/channels/flow/core/boot.py
new file mode 100644
--- /dev/null
+++ b/Engine/Extras/ushell/channels/flow/core/boot.py
@@ -0,0 +1,79 @@
+# Copyright Epic Games, Inc. All Rights Reserved.
+
+import os
+import sys
+import flow.cmd
+
+#-------------------------------------------------------------------------------
+class Boot(flow.cmd.Cmd):
+    """ Launches a new session """
+    startdir    = flow.cmd.Arg("", "Working directory to start the shell in")
+    bootarg     = flow.cmd.Opt("", "(for internal use)")
+    cleanprompt = flow.cmd.Opt(False, "Do not modify the host shell's prompt")
+
+    def main(self):
+        # Clear the screen
+        sys.stdout.write("\x1b[2J\x1b[0;0H")
+
+        # Change to the start directory the user requested. Nice and early.
+        if self.args.startdir:
+            try: os.chdir(self.args.startdir)
+            except: pass
+
+        # Register the available shells that subclasses of Boot provide
+        class ShellRegistrar(object):
+            def __init__(self):       self.shells = {}
+            def add(self, name, cls): self.shells[name] = cls
+        shell_registrar = ShellRegistrar()
+        self.register_shells(shell_registrar)
+
+        channel = self.get_channel()
+        system = channel.get_system()
+        working_dir = system.get_working_dir()
+
+        # Prepare the environment to boot into.
+        exec_context = self.get_exec_context()
+        env = exec_context.get_env()
+
+        def read_path_pieces():
+            for piece in env["PATH"].split(os.pathsep):
+                if "flow" not in piece and "/shims" not in piece:
+                    yield piece
+
+        shims_path = os.path.normpath(working_dir) + "/shims"
+        env["PATH"] = os.pathsep.join((shims_path, *read_path_pieces()))
+
+        # In case someone wants to run a shim.
+        os.environ["PATH"] = env["PATH"]
+
+        # Set a session id and somewhere to store per-session information
+        session_id = str(os.getpid())
+        env["FLOW_SID"] = session_id
+        os.environ["FLOW_SID"] = session_id
+
+        # Get the prompt for the flow shell
+        if not self.args.cleanprompt and (prompt := self.get_prompt()):
+            env["FLOW_PROMPT"] = prompt
+
+        # Go. This is more of a "set_env()" than a run really
+        self.run(env)
+
+        # Find the shell that should be used to hold this session
+        shell_name, shell_cookie = (*self.args.bootarg.split(",", 1), "")[:2]
+        try:
+            shell_class = shell_registrar.shells[shell_name]
+        except KeyError:
+            raise RuntimeError(f"Unable to find a shell named '{shell_name}'")
+
+        # Run the shell
+        system = channel.get_system()
+        shell_class(system).boot_shell(env, shell_cookie)
+
+    def get_prompt(self):
+        pass
+
+    def run(self, env):
+        pass
+
+    def register_shells(self, registrar):
+        pass
diff --git a/Engine/Extras/ushell/channels/flow/core/cmds/complete.py b/Engine/Extras/ushell/channels/flow/core/cmds/complete.py
new file mode 100644
--- /dev/null
+++ b/Engine/Extras/ushell/channels/flow/core/cmds/complete.py
@@ -0,0 +1,176 @@
+# Copyright Epic Games, Inc. All Rights Reserved.
+
+import os
+import sys
+import flow.cmd
+import importlib.util as imputil
+
+#-------------------------------------------------------------------------------
+def _set_title_inner(t, s):
+    print("\x1b]2;", t, " - ushell", s, sep=None, file=sys.stdout)
+
+if os.name == "nt":
+    def _set_title(t):
+        _set_title_inner(t, "\x1b\x5c")
+else:
+    def _set_title(t):
+        _set_title_inner(t, "\x07")
+
+#-------------------------------------------------------------------------------
+def _reply(line):   sys.stdout.write(line + "\n")
+def _reply_ok():    sys.stdout.write("\x01\n")
+def _reply_error(): sys.stdout.write("\x02\n")
+
+#-------------------------------------------------------------------------------
+def _get_prompt(system, line_iter, reply):
+    cmd_tree = system.get_command_tree()
+    prompt_node = cmd_tree.get_root_node().find_child("$prompt")
+    if not prompt_node:
+        return
+
+    if not prompt_node.is_invokeable():
+        return
+
+    prompt_class = prompt_node.get_command_class()
+    if not prompt_class:
+        return
+
+    context = prompt_class.get_class_type().Context()
+
+    prompt = prompt_class.construct()
+    prompt.prompt(context)
+
+    tags = {k:v for k,v in context.read_tags()}
+    def expand_tags(line):
+        for key, value in tags.items():
+            line = line.replace(f":{key}:", value)
+        return line
+
+    if title := (os.getenv("FLOW_TITLE") or getattr(context, "TITLE", None)):
+        title = expand_tags(title)
+        _set_title(title)
+
+    for line in line_iter:
+        line = expand_tags(line)
+        _reply(line)
+
+#-------------------------------------------------------------------------------
+def _get_tags(system, reply):
+    cmd_tree = system.get_command_tree()
+    prompt_node = cmd_tree.get_root_node().find_child("$prompt")
+    if not prompt_node:
+        return
+
+    if not prompt_node.is_invokeable():
+        return
+
+    prompt_class = prompt_node.get_command_class()
+    if not prompt_class:
+        return
+
+    context = prompt_class.get_class_type().Context()
+
+    prompt = prompt_class.construct()
+    prompt.prompt(context)
+
+    for k,v in context.read_tags():
+        _reply(k)
+        _reply(v)
+
+#-------------------------------------------------------------------------------
+def _serve(system, file=sys.stdin):
+    def read_lines():
+        inner = (x.rstrip("\r\n") for x in file)
+        for line in inner:
+            if line == "\x01":
+                break
+
+        for line in inner:
+            if line == "\x02":
+                break
+            yield line
+
+    delim = " "
+    lines = []
+    for line in read_lines():
+        if line.startswith("\x03"):
+            delim = line[1:]
+        else:
+            lines.append(line)
+
+    _reply_ok() if _complete(system, lines, _reply, delim) else _reply_error()
+
+#-------------------------------------------------------------------------------
+def _complete(system, word_iter, reply, delim):
+    words = list(word_iter)
+
+    if words and words[0] == "$":
+        _get_prompt(system, words[1:], reply)
+        return True
+
+    if words and words[0] == "$$":
+        _get_tags(system, reply)
+        return True
+
+    word_iter = (x for x in words[:-1])
+    cmd_tree = system.get_command_tree()
+    tree_node, tail_args = cmd_tree.find(word_iter)
+
+    tail_args = list(tail_args)
+
+    ok = False
+    if not tail_args:
+        for name,_ in tree_node.read_children():
+            reply(name)
+            ok = True
+
+    command_class = tree_node.get_command_class()
+    if not command_class:
+        return True
+
+    try:
+        command = command_class.construct()
+        for result in command.read_completion(tail_args, prefix=words[-1], delim=delim):
+            reply(str(result))
+    except LookupError:
+        if not ok:
+            return False
+    except:
+        pass
+
+    return True
+
+
+
+#-------------------------------------------------------------------------------
+class Complete(flow.cmd.Cmd):
+    """ Provides support for adding completion to host shells and for debugging
+    cmds' completion routines. Suffix the last word of 'words' with "..." to use
+    that as a prefix. Use "$" to print the prompt """
+    words   = flow.cmd.Arg([str], "Words to parse and print completion for")
+    daemon  = flow.cmd.Opt(False, "Run as a daemon serving requests over stdin/out")
+    delim   = flow.cmd.Opt(" ", "The character that delimits the last word")
+
+    def _main_daemon(self, system):
+        primary_dir = os.getcwd()
+        try:
+            while True:
+                _serve(system)
+                os.chdir(primary_dir)
+        except:
+            pass
+
+    def main(self):
+        channel = self.get_channel()
+        system = channel.get_system()
+
+        if self.args.daemon:
+            return self._main_daemon(system)
+
+        words = self.args.words
+        if words and words[-1].endswith("..."):
+            word_iter = (*words[:-1], words[-1][:-3])
+        else:
+            word_iter = (*words, "")
+
+        return _complete(system, word_iter, print, self.args.delim)
diff --git a/Engine/Extras/ushell/channels/flow/core/cmds/debug.py b/Engine/Extras/ushell/channels/flow/core/cmds/debug.py
new file mode 100644
--- /dev/null
+++ b/Engine/Extras/ushell/channels/flow/core/cmds/debug.py
@@ -0,0 +1,160 @@
+# Copyright Epic Games, Inc. All Rights Reserved.
+
+import os
+import re
+import sys
+import pprint
+import hashlib
+import marshal
+import flow.cmd
+from urllib.request import urlopen, URLError
+
+#-------------------------------------------------------------------------------
+def _http_get(url, on_data):
+    # Fire up an http client to get the url.
+    client = urlopen(url)
+    if (client.status // 100) != 2:
+        assert False, f"Error creating HTTP client for {url}"
+
+    # Get the download's file name
+    content_header = client.headers.get("content-disposition", "")
+    m = re.search(r'filename="?([^;"]+)"?(;|$)', content_header)
+    if m: file_name = m.group(1)
+    else: file_name = os.path.basename(client.url)
+
+    # Do the download.
+    recv_size = 0
+    total_size = int(client.headers.get("content-length"))
+    while chunk := client.read(128 << 10):
+        recv_size += len(chunk)
+        on_data(file_name, recv_size, total_size, chunk)
+
+
+
+#-------------------------------------------------------------------------------
+class Debug(flow.cmd.Cmd):
+    """ Shows debug information about the internal state of the session """
+
+    def main(self):
+        channel = self.get_channel()
+        system = channel.get_system()
+
+        with open(system.get_working_dir() + "manifest", "rb") as x:
+            manifest = marshal.load(x)
+
+        cmd_tree = manifest["cmd_tree"]
+        pprint.pprint(manifest)
+
+        print()
+        session = self.get_noticeboard(self.Noticeboard.SESSION)
+        pprint.pprint(session._get_inner())
+
+#-------------------------------------------------------------------------------
+class Invalidate(flow.cmd.Cmd):
+    """ Invalidates cached state such that initialisation runs again the next
+    time the shell is started up."""
+
+    def main(self):
+        from pathlib import Path
+
+        channel = self.get_channel()
+        system = channel.get_system()
+        working_dir = Path(system.get_working_dir())
+        print(working_dir)
+
+        for manifest_path in working_dir.glob("**/manifest"):
+            if not manifest_path.parent.samefile(working_dir):
+                print("Removing:", manifest_path.relative_to(working_dir))
+                manifest_path.unlink(missing_ok=True)
+
+#-------------------------------------------------------------------------------
+class Nuke(flow.cmd.Cmd):
+    """ Wipes everything. """
+
+    def main(self):
+        import shutil
+        channel = self.get_channel()
+        system = channel.get_system()
+        for item in (system.get_temp_dir(), system.get_working_dir()):
+            print("Removing:", item, end="")
+            shutil.rmtree(item, ignore_errors=True)
+            print("")
+
+        from pathlib import Path
+        tool_dir = Path(system.get_tools_dir())
+        print("Tools dir:", tool_dir)
+        for item in tool_dir.glob("*"):
+            if item.is_dir():
+                print("Removing tool:", item.name, end="")
+                item = item.rename(str(item) + str(os.getpid()))
+                shutil.rmtree(item, ignore_errors=True)
+                print("...error" if item.is_dir() else "")
+
+#-------------------------------------------------------------------------------
+class Paths(flow.cmd.Cmd):
+    """ Prints behind-the-scenes paths. """
+
+    def main(self):
+        channel = self.get_channel()
+        system = channel.get_system()
+        print("State:", system.get_working_dir())
+        print("Tools:", system.get_tools_dir())
+        print(" Temp:", system.get_temp_dir())
+
+#-------------------------------------------------------------------------------
+class _Channels(flow.cmd.Cmd):
+    def main(self):
+        channel = self.get_channel()
+        system = channel.get_system()
+
+        self.channels = {}
+        channels_dir = system.get_working_dir() + "/channels/"
+        with os.scandir(channels_dir) as dirs:
+            for item in dirs:
+                with open(item.path + "/manifest", "rb") as data:
+                    manifest = marshal.load(data)
+
+                self.channels[item.name] = manifest["tools"]
+
+#-------------------------------------------------------------------------------
+class Tools(_Channels):
+    """ Prints debug information about tools. """
+
+    def main(self):
+        super().main()
+        pprint.pprint(self.channels)
+
+#-------------------------------------------------------------------------------
+class Sha1s(_Channels):
+    """ Shows channels' tools SHA1s """
+
+    def _impl(self, name, manifest):
+        print("--", name)
+        for tool_name, tool in manifest.items():
+            def on_data(file_name, recv_size, total_size, data):
+                if sys.stdout.isatty():
+                    print("%.02f%%" % ((recv_size * 100) / total_size), end="\r")
+                digest.update(data)
+
+            digest = hashlib.sha1()
+            for urls in tool["bundles"].values():
+                for url in urls:
+                    try: _http_get(url, on_data)
+                    except URLError: pass
+            sha1 = digest.hexdigest()
+
+            print(rf'{tool_name}.sha1("{sha1}")')
+
+    def main(self):
+        super().main()
+        for name, manifest in self.channels.items():
+            self._impl(name, manifest)
+
+#-------------------------------------------------------------------------------
+class Argumentss(flow.cmd.Cmd):
+    """ Prints the given arguments as they would be received by commands """
+    arguments = flow.cmd.Arg([str], "Arguments to print")
+
+    def main(self):
+        for i, arg in enumerate(self.args.arguments):
+            print("%-2d:" % i, arg)
diff --git a/Engine/Extras/ushell/channels/flow/core/cmds/help.py b/Engine/Extras/ushell/channels/flow/core/cmds/help.py
new file mode 100644
--- /dev/null
+++ b/Engine/Extras/ushell/channels/flow/core/cmds/help.py
@@ -0,0 +1,95 @@
+# Copyright Epic Games, Inc. All Rights Reserved.
+
+import re
+import sys
+import shutil
+import marshal
+import flow.cmd
+import textwrap
+import importlib
+from pathlib import Path
+
+#-------------------------------------------------------------------------------
+class Help(flow.cmd.Cmd):
+    """ Prints information about the available commands """
+
+    def main(self):
+        channel = self.get_channel()
+        system = channel.get_system()
+        cmd_tree = system.get_command_tree()
+        root_node = cmd_tree.get_root_node()
+
+        def read_invokeables(node, path=""):
+            if node.is_invokeable():
+                yield path, node
+
+            for child_name, child_node in node.read_children():
+                child_path = (path + " " + child_name) if path else child_name
+                yield from read_invokeables(child_node, child_path)
+
+        def read_output():
+            for path, node in read_invokeables(root_node):
+                if path.startswith("$"):
+                    continue
+
+                try:
+                    command_class = node.get_command_class()
+                    command_class = command_class.get_class_type()
+                    if hasattr(command_class, "get_desc"):
+                        yield path, command_class.get_desc().strip()
+                except Exception as e:
+                    yield path, "LOAD_ERROR: " + str(e)
+
+        right_column = shutil.get_terminal_size()[0]
+        right_column -= 5
+
+        output = list(read_output())
+        output.sort(key=lambda x: x[0])
+
+        name_len = max(len(x[0]) for x in output) + 4
+        whitespace = re.compile(r"\s+")
+        desc_width = right_column - name_len
+        for name, desc in output:
+            dots = "." * (name_len - len(name) - 3)
+            initial_prefix = flow.cmd.text.light_cyan(name) + " " + flow.cmd.text.grey(dots)
+
+            desc = whitespace.sub(" ", desc)
+            desc_lines = textwrap.wrap(desc, desc_width)
+
+            first_line = 0
+            for line in desc_lines:
+                prefix = initial_prefix if not first_line else (" " * (name_len - 2))
+                first_line += 1
+                print("", prefix, line)
+
+            if first_line > 1:
+                print()
+
+
+
+#-------------------------------------------------------------------------------
+class ReadMe(flow.cmd.Cmd):
+    """ Opens the README.txt file in your editor of choice or prints it to stdout """
+    editor = flow.cmd.Opt(False, "Open in a text editor")
+
+    def main(self):
+        for haystack_dir in Path(__file__).parents:
+            readme_path = haystack_dir / "README.txt"
+            if readme_path.is_file():
+                break
+        else:
+            self.print_error("Unable to find README.txt")
+            return False
+
+        if self.args.editor:
+            self.edit_file(readme_path)
+            return
+
+        with readme_path.open("rt") as file:
+            for line in file.readlines():
+                line = line.strip()
+                if line.startswith("#"):
+                    with flow.cmd.text.light_cyan:
+                        print(line)
+                else:
+                    print(line)
diff --git a/Engine/Extras/ushell/channels/flow/core/cmds/tips.py b/Engine/Extras/ushell/channels/flow/core/cmds/tips.py
new file mode 100644
--- /dev/null
+++ b/Engine/Extras/ushell/channels/flow/core/cmds/tips.py
@@ -0,0 +1,41 @@
+# Copyright Epic Games, Inc. All Rights Reserved.
+
+import random
+import flow.cmd
+
+#-------------------------------------------------------------------------------
+class Tips(flow.cmd.Cmd):
+    """ Displays a useful tip """
+
+    @classmethod
+    def read_get_tip_methods(self):
+        for component in self.__mro__:
+            if method := getattr(component, "get_tips", None):
+                yield method
+
+    def main(self):
+        tips = []
+        for get_tips in self.read_get_tip_methods():
+            tips += get_tips(self)
+
+        with flow.cmd.text.light_cyan:
+            print("TIP!", end="")
+
+        with flow.cmd.text.white:
+            print("", tips[random.randrange(len(tips))])
+
+    def get_tips(self):
+        return (
+            "Press Tab for context sensitive completion",
+            "List command history using 'history' and use ![NUM] to recall.",
+            "Incrementally search your command history with Ctrl-R (or Ctrl-S)",
+            "'.[command] --help' to see a commands inline help",
+            "Run '.help' to see all available commands",
+            "Calvin and Hobbes is excellent",
+            "Customise a session with '$USERPROFILE/.ushell/hooks/startup.bat'",
+            "'d:\>dc' changes directory using the fuzzy finder",
+            "'d:\>dc.' limits fuzzy-finder-chdir to the current directory",
+            "'d:\>dc DIR' limits fuzzy-finder-chdir to directory 'DIR'",
+            "Use Ctrl-P/Ctrl-N to recall previous queries to 'qq' and 'dc'",
+            "'qq' to fuzzy find and execute lines from the command history",
+        )
diff --git a/Engine/Extras/ushell/channels/flow/core/describe.flow.py b/Engine/Extras/ushell/channels/flow/core/describe.flow.py
new file mode 100644
--- /dev/null
+++ b/Engine/Extras/ushell/channels/flow/core/describe.flow.py
@@ -0,0 +1,85 @@
+# Copyright Epic Games, Inc. All Rights Reserved.
+
+import os
+import flow.describe
+
+#-------------------------------------------------------------------------------
+help_cmd = flow.describe.Command()
+help_cmd.source("cmds/help.py", "Help")
+help_cmd.invoke("help")
+
+help_readme = flow.describe.Command()
+help_readme.source("cmds/help.py", "ReadMe")
+help_readme.invoke("help", "readme")
+
+#-------------------------------------------------------------------------------
+tip_cmd = flow.describe.Command()
+tip_cmd.source("cmds/tips.py", "Tips")
+tip_cmd.invoke("tip")
+tip_cmd.prefix("$")
+
+#-------------------------------------------------------------------------------
+debug_manifest = flow.describe.Command()
+debug_manifest.source("cmds/debug.py", "Debug")
+debug_manifest.invoke("debug", "manifest")
+debug_manifest.prefix("$")
+
+#-------------------------------------------------------------------------------
+debug_tools = flow.describe.Command()
+debug_tools.source("cmds/debug.py", "Tools")
+debug_tools.invoke("debug", "tools")
+debug_tools.prefix("$")
+
+#-------------------------------------------------------------------------------
+debug_sha1s = flow.describe.Command()
+debug_sha1s.source("cmds/debug.py", "Sha1s")
+debug_sha1s.invoke("debug", "tools", "sha1s")
+debug_sha1s.prefix("$")
+
+#-------------------------------------------------------------------------------
+debug_args = flow.describe.Command()
+debug_args.source("cmds/debug.py", "Argumentss")
+debug_args.invoke("debug", "args")
+debug_args.prefix("$")
+
+#-------------------------------------------------------------------------------
+debug_invalidate = flow.describe.Command()
+debug_invalidate.source("cmds/debug.py", "Invalidate")
+debug_invalidate.invoke("debug", "invalidate")
+debug_invalidate.prefix("$")
+
+#-------------------------------------------------------------------------------
+debug_paths = flow.describe.Command()
+debug_paths.source("cmds/debug.py", "Paths")
+debug_paths.invoke("debug", "paths")
+debug_paths.prefix("$")
+
+#-------------------------------------------------------------------------------
+debug_nuke = flow.describe.Command()
+debug_nuke.source("cmds/debug.py", "Nuke")
+debug_nuke.invoke("debug", "nuke")
+debug_nuke.prefix("$")
+
+#-------------------------------------------------------------------------------
+complete = flow.describe.Command()
+complete.source("cmds/complete.py", "Complete")
+complete.invoke("complete")
+complete.prefix("$")
+
+#-------------------------------------------------------------------------------
+prompt = flow.describe.Command()
+prompt.source("prompt.py", "Prompt")
+prompt.invoke("prompt")
+prompt.prefix("$")
+
+#-------------------------------------------------------------------------------
+boot = flow.describe.Command()
+boot.source("boot.py", "Boot")
+boot.invoke("boot")
+boot.prefix("$")
+
+
+
+#-------------------------------------------------------------------------------
+channel = flow.describe.Channel()
+channel.version("1")
diff --git a/Engine/Extras/ushell/channels/flow/core/prompt.py b/Engine/Extras/ushell/channels/flow/core/prompt.py
new file mode 100644
--- /dev/null
+++ b/Engine/Extras/ushell/channels/flow/core/prompt.py
@@ -0,0 +1,33 @@
+# Copyright Epic Games, Inc. All Rights Reserved.
+
+import os
+import flow.cmd
+
+class Prompt(flow.cmd.Cmd):
+    """ Used to collect key/value pairs for tags in the prompt. """
+    format = flow.cmd.Opt("none", "Output format. Options; none,sh")
+
+    class Context(object):
+        def read_tags(self):
+            for k in dir(self):
+                v = getattr(self, k)
+                if not k.startswith("_") and isinstance(v, str):
+                    yield k, v
+
+    def main(self):
+        context = Prompt.Context()
+        self.prompt(context)
+
+        # Should we format the output in a particular way?
+        if self.args.format == "sh":
+            print(*(f"{k}='{v}'" for k, v in context.read_tags()), end="")
+            return
+
+        # Just print the tags in a friendly manner
+        for key, value in context.read_tags():
+            print(key, "=", value)
+
+    def prompt(self, context):
+        cwd = os.getcwd()
+        context.FLOW_CWD = cwd
+        context.CWD = cwd
diff --git a/Engine/Extras/ushell/channels/flow/core/system/boot.py b/Engine/Extras/ushell/channels/flow/core/system/boot.py
new file mode 100644
--- /dev/null
+++ b/Engine/Extras/ushell/channels/flow/core/system/boot.py
@@ -0,0 +1,52 @@
+# Copyright Epic Games, Inc. All Rights Reserved.
+
+import os
+import sys
+
+#-------------------------------------------------------------------------------
+def _enable_vt100():
+    if os.name != "nt":
+        return
+
+    import ctypes
+    _win_dll = ctypes.LibraryLoader(ctypes.WinDLL)
+    _get_std_handle = _win_dll.kernel32.GetStdHandle
+    _get_console_mode = _win_dll.kernel32.GetConsoleMode
+    _set_console_mode = _win_dll.kernel32.SetConsoleMode
+
+    con_mode = ctypes.c_int()
+    stdout_handle = _get_std_handle(-11)
+    _get_console_mode(stdout_handle, ctypes.byref(con_mode))
+    con_mode = con_mode.value | 4
+    _set_console_mode(stdout_handle, con_mode)
+
+#-------------------------------------------------------------------------------
+def _main():
+    # Flow is started like this; boot.py [working_dir] [channel_dir] -- [shell_args]
+    boot_dirs = []
+    shell_args = []
+    collector = boot_dirs
+    for arg in sys.argv[1:]:
+        if arg == "--":
+            collector = shell_args
+            continue
+
+        collector.append(arg)
+
+    sys.path.append(os.path.abspath(__file__ + "/../lib"))
+
+    # Get all the channels ready and aggregated
+    import bootstrap
+    working_dir, *channels_dirs = boot_dirs
+    working_suffix = bootstrap.impl(working_dir, *channels_dirs)
+
+    # Run the '$boot' command
+    import run
+    ret = run.main((f"{working_dir}/{working_suffix}/manifest", "$boot", *shell_args))
+    raise SystemExit(ret)
+
+
+
+if __name__ == "__main__":
+    _enable_vt100()
+    _main()
diff --git a/Engine/Extras/ushell/channels/flow/core/system/flow/_backtrace.py b/Engine/Extras/ushell/channels/flow/core/system/flow/_backtrace.py
new file mode 100644
--- /dev/null
+++ b/Engine/Extras/ushell/channels/flow/core/system/flow/_backtrace.py
@@ -0,0 +1,47 @@
+# Copyright Epic Games, Inc. All Rights Reserved.
+
+import sys
+
+#-------------------------------------------------------------------------------
+def backtrace(exc_type, exc_value, tb):
+    from textwrap import wrap
+    from os.path import normpath
+
+    def red_text(x):    return "\x1b[91m" + x + "\x1b[0m"
+    def white_text(x):  return "\x1b[97m" + x + "\x1b[0m"
+
+    max_func_len = 0
+    frames = []
+    while tb:
+        frame = tb.tb_frame
+        code = frame.f_code
+        max_func_len = max(len(code.co_name), max_func_len)
+        frames.insert(0, (code.co_name, normpath(code.co_filename), frame.f_lineno))
+        tb = tb.tb_next
+
+    def spam(x=""):
+        sys.stderr.write(str(red_text("##")) + " " + str(x) + "\n")
+
+    sys.stderr.write("\n")
+    message = str(exc_value)
+    if len(message):
+        spam()
+        for line in wrap(str(message), 75):
+            spam(white_text(line))
+    spam()
+    spam(red_text("[%s]") % exc_type.__name__)
+    spam()
+
+    max_file_len = max(16, 55 - max_func_len)
+    for func, file, line in frames[:4]:
+        if len(file) > max_file_len:
+            file = "_" + file[-max_file_len:]
+        dots = "." * (max_func_len - len(func))
+        x = "%s ..%s %s:%d" % (func, dots, file, line)
+        spam(red_text(x))
+
+    spam()
+
+
+
+sys.excepthook = backtrace
diff --git a/Engine/Extras/ushell/channels/flow/core/system/flow/_flick.py b/Engine/Extras/ushell/channels/flow/core/system/flow/_flick.py
new file mode 100644
--- /dev/null
+++ b/Engine/Extras/ushell/channels/flow/core/system/flow/_flick.py
@@ -0,0 +1,389 @@
+# Copyright Epic Games, Inc. All Rights Reserved.
+
+#-------------------------------------------------------------------------------
+class _ArgOptBase(object):
+    def __init__(self, type_value, description):
+        self._type_value = type_value
+        self._description = description
+
+    def get_desc(self):
+        return self._description
+
+    def get_default(self):
+        inner = self._type_value
+        return None if isinstance(inner, type) else inner
+
+    def get_type(self):
+        inner = self._type_value
+        return inner if isinstance(inner, type) else type(inner)
+
+    def cast_value(self, value):
+        try:
+            return None if value == None else self.get_type()(value)
+        except ValueError:
+            raise TypeError(self, value)
+
+    def validate(self):
+        if not self.get_desc():
+            style = "Argument" if "Arg" in type(self).__name__ else "Option"
+            raise SyntaxError(style + " has no description")
+
+#-------------------------------------------------------------------------------
+class _ArgBase(_ArgOptBase):
+    def validate(self):
+        super().validate()
+        if self.get_type() == bool:
+            raise SyntaxError("Positional arguments cannot be booleans")
+
+        if self.get_type() == tuple:
+            raise SyntaxError("Positional arguments cannot be tuples")
+
+#-------------------------------------------------------------------------------
+class _Arg(_ArgBase):
+    pass
+
+#-------------------------------------------------------------------------------
+class _ManyArg(_ArgBase):
+    def get_default(self):
+        return ()
+
+    def validate(self):
+        super().validate()
+        if type(self._type_value) != type:
+            raise SyntaxError("Array-type arguments must specify a type (e.g. [str])")
+
+#-------------------------------------------------------------------------------
+class _Opt(_ArgOptBase):
+    def get_value(self):
+        return None
+
+    def needs_param(self):
+        return self.get_type() != bool
+
+    def validate(self):
+        super().validate()
+        if isinstance(self._type_value, type):
+            raise SyntaxError("Optional arguments must have a default value")
+
+        #if self.is_multi():
+            #raise SyntaxError("Optional arguments cannot be array-types")
+
+#-------------------------------------------------------------------------------
+class _TriOpt(_Opt):
+    def __init__(self, type_value, *args):
+        super().__init__(type_value[0], *args)
+        self._value = type_value[1]
+
+    def get_value(self):
+        return self._value
+
+    def get_type(self):
+        return type(self._value)
+
+    def needs_param(self):
+        return True
+
+
+
+#-------------------------------------------------------------------------------
+def Arg(type_value, *args):
+    if type(type_value) == list:
+        if len(type_value) > 0:
+            type_value = type_value[0]
+        return _ManyArg(type_value, *args)
+    return _Arg(type_value, *args)
+
+def Opt(type_value, *args):
+    if type(type_value) == tuple:
+        return _TriOpt(type_value, *args)
+    return _Opt(type_value, *args)
+
+
+
+#-------------------------------------------------------------------------------
+class _ArgOptBox(object):
+    def __init__(self, items):
+        super().__setattr__("_items", items)
+
+    def __iter__(self):
+        for name, item in self._items.items():
+            yield name, item[0]
+
+    def __getattr__(self, name):
+        return self._items[name][0]
+
+    def __setattr__(self, name, value):
+        self._items[name] = (value, False)
+
+    def is_default(self, name):
+        return self._items[name][1]
+
+
+
+#-------------------------------------------------------------------------------
+class Cmd(object):
+    @classmethod
+    def _read_argopts(cls, read_type):
+        for base in (x for x in cls.__mro__ if x != object or x != Cmd):
+            items = base.__dict__.items()
+            for name, info in (x for x in items if isinstance(x[1], read_type)):
+                if not name.startswith("_"):
+                    yield name, info
+
+    def _print_error(self, message):
+        print(message)
+
+    def _print_help(self, short=False):
+        # Description
+        if not short:
+            for line in self.get_desc().split("\n"):
+                line = line[:4].strip() + line[4:]
+                print(line)
+            print()
+
+        # Usage
+        args = []
+        print("USAGE:", end="" if short else "\n")
+        print("", "[OPTIONS]", end="")
+        for arg_name, arg in self.read_args():
+            if arg.get_default() is None:
+                print("", arg_name, end="")
+            else:
+                print("", f"[{arg_name}]", end="")
+            args.append((arg_name, arg.get_desc()))
+
+        if arg_many := self.get_many_arg():
+            print(f" [[--] {arg_many[0]}...]", end="")
+            args.append(("--", f"Delineates where '{arg_many[0]}' starts (optional)"))
+            args.append((arg_many[0], arg_many[1].get_desc()))
+        print()
+
+        if short:
+            return
+
+        # Options
+        opts = []
+        for opt_name, opt in self.read_opts():
+            if opt.needs_param():
+                value = opt.get_value()
+                if value is None:
+                    opt_name += " " + opt.get_type().__name__.upper()
+                else:
+                    opt_name += f"[={opt.get_type().__name__.upper()}]"
+            opts.append(("--" + opt_name, opt.get_desc()))
+        opts.append(("--help", "Display this help text"))
+
+        # Output
+        col_a = 0 if not args else max(len(x[0]) for x in args)
+        col_o = max(len(x[0]) for x in opts)
+        col_0 = max(col_a, col_o)
+
+        once = True
+        for name, desc in args:
+            if once:
+                print("\nARGS:")
+                once = False
+            print("  %-*s" % (col_0, name), desc)
+
+        print("\nOPTIONS:")
+        for name, desc in opts:
+            print("  %-*s" % (col_0, name), desc)
+
+    def _call_main(self):
+        ret = self.main()
+
+        # Run main() in an asyncio event loop if 'ret' is a coroutine
+        if getattr(ret, "cr_code", None):
+            import asyncio
+            ret = asyncio.run(ret)
+
+        return ret
+
+    def try_parse_args(self, result, *in_args):
+        opts = {k:v for k,v in self.read_opts()}
+
+        # Set defaults
+        result.update({k:(v.get_default(), True) for k,v in opts.items()})
+        for arg_name, cmd_arg in self._read_argopts(_ArgBase):
+            result[arg_name] = (cmd_arg.get_default(), True)
+
+        # Parses an opt-type from input stream
+        def parse_opt(arg, in_args_iter):
+            # Parse "--opt_name[=value]"
+            opt_name, value = (*arg[2:].split("=", 1), None)[:2]
+
+            # Do we know of this here option thing?
+            opt = opts.get(opt_name, None)
+            if not opt:
+                raise ValueError((opt_name,))
+
+            # Pull out the value from the next in_arg (or use =value from above)
+            if opt.needs_param():
+                value = value if value != None else opt.get_value()
+                if value == None:
+                    value = next(in_args_iter, None)
+                    if not value or value.startswith("-"):
+                        raise LookupError(opt_name, opt)
+            else:
+                # "Not" because only bool-type ops can have no parameter, and an
+                # _Opt(False, "desc") has a default of False but is True if parsed
+                value = not opt.get_default()
+
+            result[opt_name] = (opt.cast_value(value), False)
+
+        # Parse Opts and separate out positional arguments
+        positionals = []
+        in_args_iter = iter(in_args)
+        for arg in in_args_iter:
+            if arg.startswith("--"):
+                if len(arg) == 2:
+                    positionals.append("--")
+                    positionals.extend(in_args_iter)
+                    break
+                parse_opt(arg, in_args_iter)
+            else:
+                positionals.append(arg)
+
+        # Parse Args
+        index = -1
+        cmd_arg = None
+        cmd_args = self.read_args()
+        for index, (arg_name, cmd_arg) in enumerate(cmd_args):
+            if index >= len(positionals):
+                return iter(((arg_name, cmd_arg), *cmd_args))
+
+            in_arg = positionals[index]
+            if in_arg == "--":
+                break
+
+            result[arg_name] = (cmd_arg.cast_value(in_arg), False)
+        index += 1
+
+        # Trim off a leading "--"
+        if len(positionals) > index and positionals[index] == "--":
+            index += 1
+
+        # Is there a many-arg sink for what remains?
+        many_cmd_arg = self.get_many_arg()
+        if many_cmd_arg:
+            arg_name, cmd_arg = many_cmd_arg
+            result[arg_name] = (tuple(positionals[index:]), len(positionals[index:]) == 0)
+            return iter(((arg_name, cmd_arg),))
+
+        # Check if there's anything left unparsed
+        if positionals[index:]:
+            raise ValueError(positionals[index:])
+
+    def read_args(self):
+        yield from self._read_argopts(_Arg)
+
+    def read_opts(self):
+        yield from self._read_argopts(_Opt)
+
+    def get_many_arg(self):
+        return next(self._read_argopts(_ManyArg), None)
+
+    @classmethod
+    def get_name(self):
+        return self.__name__
+
+    @classmethod
+    def get_desc(self):
+        it = (x.__doc__ for x in self.__mro__ if x.__doc__ and x != object)
+        return "\n\n".join(it)
+
+    def invoke(self, in_args):
+        self.validate()
+
+        def read_in_args():
+            double_dashed = False
+            for arg in in_args:
+                double_dashed |= (arg == "--")
+                if arg in ("--help", "-h", "/?"):
+                    if not double_dashed:
+                        self._print_help()
+                        raise SystemExit(127)
+                else:
+                    yield arg
+
+        def on_error(message):
+            self._print_help(short=True)
+            self._print_error("ERROR: " + message)
+            raise SystemExit(126)
+
+        args_out = {}
+        try:
+            if cmd_args := self.try_parse_args(args_out, *read_in_args()):
+                for arg_name, cmd_arg in cmd_args:
+                    if cmd_arg.get_default() == None:
+                        raise LookupError(arg_name, cmd_arg)
+        except LookupError as e:
+            name, argopt = e.args
+            on_error(f"Missing value for argument '{name}'")
+        except ValueError as e:
+            arg = " ".join(e.args[0])
+            on_error(f"Unknown argument(s) '{arg}'")
+        except TypeError as e:
+            argopt, value = e.args
+            out_type = argopt.get_type().__name__
+            on_error(f"Unable to convert '{value}' to type '{out_type}'")
+
+        self.args = _ArgOptBox(args_out)
+
+        return self._call_main()
+
+    def read_completion(self, in_args, prefix="", delim=" "):
+        in_args = tuple(in_args)
+
+        if "--" not in in_args and prefix.startswith("-"):
+            yield from ("--" + k for k,v in self.read_opts())
+            yield "--help"
+            return
+
+        args_out = {}
+        try:
+            arg_iter = self.try_parse_args(args_out, *in_args)
+        except ValueError:
+            raise LookupError
+        except LookupError as e:
+            arg_iter = iter((e.args,))
+        except:
+            arg_iter = None
+        if not arg_iter:
+            return
+
+        arg_name, argopt = next(arg_iter)
+
+        completer = getattr(self, "complete_" + arg_name, None)
+        if completer == None:
+            raise LookupError
+
+        if not completer:
+            return
+
+        if hasattr(completer, "__call__"):
+            self.args = _ArgOptBox(args_out)
+            completer = completer(prefix)
+
+        if completer:
+            yield from iter(completer)
+
+    def validate(self):
+        if not hasattr(self, "main"):
+            raise SyntaxError(f"Command '{self.get_name()}' has no main() method")
+
+        if not self.get_desc():
+            raise SyntaxError(f"Command '{self.get_name()}' has no valid description")
+
+        #last_arg = None
+        #first_multi = None
+        for argopt_name, argopt in self._read_argopts(_ArgOptBase):
+            try:
+                argopt.validate()
+                #last_arg = argopt if isinstance(argopt, _Arg) else last_arg
+                #first_multi = first_multi or (argopt if argopt.is_multi() else None)
+            except SyntaxError as e:
+                raise SyntaxError(f"'{self.get_name()}.{argopt_name}': {e}")
+
+        #if first_multi and first_multi != last_arg:
+            #raise SyntaxError(f"Only the last positional argument can be an array-type")
diff --git a/Engine/Extras/ushell/channels/flow/core/system/flow/_runnable.py b/Engine/Extras/ushell/channels/flow/core/system/flow/_runnable.py
new file mode 100644
--- /dev/null
+++ b/Engine/Extras/ushell/channels/flow/core/system/flow/_runnable.py
@@ -0,0 +1,149 @@
+# Copyright Epic Games, Inc. All Rights Reserved.
+
+import os
+import subprocess as sp
+from os.path import exists, join, basename
+
+#-------------------------------------------------------------------------------
+if os.name == "nt":
+    import struct
+    def _is_gui_binary(binary_path):
+        with open(binary_path, "rb") as exe_file:
+            # DOS header
+            dos_header = exe_file.read(64)
+            dos_header = struct.unpack("H" + ("H" * 29) + "I", dos_header)
+            pe_offset = dos_header[-1]
+            assert dos_header[0] == 0x5a4d
+
+            # NT header
+            exe_file.seek(pe_offset)
+            pe_magic, = struct.unpack("I", exe_file.read(4))
+            assert pe_magic == 0x4550 # 'PE'
+
+            # File header
+            file_header = struct.unpack("HHIIIHH", exe_file.read(20))
+            optional_header_size = file_header[-2] - (8 * 16)
+
+            # Optional header
+            optional_header = exe_file.read(optional_header_size)
+            if file_header[0] == 0x8664: opt_unpack = "HBBIIIIIQIIHHHHHHIIIIHHQQQQII"
+            else: opt_unpack = "HBBIIIIIIIIIHHHHHHIIIIHHIIIIII"
+            optional_header = struct.unpack(opt_unpack, optional_header)
+            subsystem = optional_header[-8]
+            return subsystem == 2
+
+    def _read_args_impl(*args):
+        for arg in (os.fspath(x) for x in args):
+            if " " in arg:
+                arg = arg.replace('"', r'\"')
+                yield rf'"{arg}"'
+            else:
+                yield str(arg)
+else:
+    def _is_gui_binary(binary_path):
+        return False
+
+    def _read_args_impl(*args):
+        yield from (os.fspath(x) for x in args)
+
+
+
+#-------------------------------------------------------------------------------
+class Runnable(object):
+    def __init__(self, exe, *args, env=None):
+        self._proc = None
+        self._ret = None
+        self._env = env
+
+        exe = str(os.fspath(exe))
+        if not exists(exe) and "/" not in exe.replace("\\", "/"):
+            paths = os.getenv("PATH", "").split(os.pathsep)
+            x = next((x for x in paths if exists(join(x, exe))), None)
+            exe = join(x, exe) if x else exe
+        self._exe = exe
+        self._args = [x for x in args if x != None]
+
+    def __add__(self, arg):
+        if arg != None:
+            self._args.append(arg)
+        return self
+
+    def __iter__(self):
+        stdout = self.launch(stdout=True)
+        return (x.decode().rstrip() for x in iter(stdout.readline, b""))
+
+    def get_exe_path(self):
+        return self._exe
+
+    def read_args(self):
+        yield from _read_args_impl(self.get_exe_path(), *self._args)
+
+    def is_gui(self):
+        return _is_gui_binary(self._exe)
+
+    def set_env(self, env):
+        self._env = env
+
+    def launch(self, stdin=None, stdout=None, suspended=False, new_term=False):
+        create_flags = 0
+        if suspended: create_flags |= 0x4 # CREATE_SUSPENDED
+        if new_term:  create_flags |= sp.CREATE_NEW_CONSOLE
+
+        kwargs = {
+            "env" : self._env,
+            "stdin" : stdin,
+            "stdout" : sp.PIPE if stdout else None if stdout == None else sp.DEVNULL,
+            "stderr" : sp.STDOUT if stdout else None,
+            "creationflags" : create_flags,
+        }
+
+        if os.name == "nt":
+            cmd = " ".join(self.read_args())
+        else:
+            cmd = (self.get_exe_path(), *self._args)
+
+        self._proc = sp.Popen(cmd, **kwargs)
+        return self._proc.stdout
+
+    def wait(self):
+        if not self._proc:
+            return
+
+        # Wait for the process.
+        if self._proc.stdin:
+            self._proc.communicate()
+            self._proc.stdin.close()
+        self._proc.wait()
+
+        # Close pipes.
+        if self._proc.stdout: self._proc.stdout.close()
+        if self._proc.stderr: self._proc.stderr.close()
+
+        self._ret = self._proc.returncode
+        self._proc = None
+        return self._ret
+
+    def get_pid(self):
+        return self._proc.pid if self._proc else 0
+
+    def get_return_code(self):
+        return self._ret if self._ret else self.wait()
+
+    def kill(self):
+        if self._proc:
+            self._proc.kill()
+        self._proc = None
+        self._ret = None
+
+    def run(self, **kwargs):
+        self.launch(**kwargs)
+        self.wait()
+        return self._ret
+
+    def run2(self):
+        self.launch(stdout=True)
+        output = self._proc.communicate()[0].decode()
+        self._ret = self._proc.returncode
+        self._proc = None
+        return self._ret, output
+
diff --git a/Engine/Extras/ushell/channels/flow/core/system/flow/cmd.py b/Engine/Extras/ushell/channels/flow/core/system/flow/cmd.py
new file mode 100644
--- /dev/null
+++ b/Engine/Extras/ushell/channels/flow/core/system/flow/cmd.py
@@ -0,0 +1,285 @@
+# Copyright Epic Games, Inc. All Rights Reserved.
+
+import os
+import sys
+import marshal
+from enum import Enum
+
+from . import _backtrace
+from .text import *
+
+#-------------------------------------------------------------------------------
+def _lovely_info(*message):
+    with text.light_cyan:
+        print("--", *message)
+
+def _lovely_warning(*message):
+    with text.light_yellow:
+        print("!!", *message)
+
+def _lovely_error(*message):
+    with text.light_red:
+        print("!!", *message)
+
+
+
+#-------------------------------------------------------------------------------
+from . import _runnable
+class _Runnable(_runnable.Runnable):
+    def launch(self, **kwargs):
+        if "silent" not in kwargs:
+            with text.cyan:
+                print("== Run: ", end="")
+            with text.white:
+                it = iter(self.read_args())
+                print(os.fspath(next(it, "")).replace("/", os.sep), end=" ")
+                print(*it)
+        else:
+            del kwargs["silent"]
+
+        return super().launch(**kwargs)
+
+#-------------------------------------------------------------------------------
+class _Env(object):
+    def __init__(self):
+        try:
+            import nt
+            self._original = nt.environ
+            if "PROMPT" in self._original:
+                del self._original["PROMPT"]
+        except ModuleNotFoundError:
+            import os
+            self._original = os.environ
+        self._env = self._original.copy()
+        self._keys = {x.upper():x for x in self._env.keys()}
+        self._edits = set()
+
+    def __setitem__(self, key, value):
+        key = self._map_key(key, True)
+        self._env[key] = str(value)
+        self._edits.add(key)
+
+    def __getitem__(self, key):
+        key = self._map_key(key, False)
+        return self._env[key]
+
+    def __delitem__(self, key):
+        key = self._map_key(key, False)
+        del self._env[key]
+        del self._keys[key.upper()]
+
+    def _map_key(self, key, add):
+        ret = self._keys.get(key.upper(), None);
+        if not ret and add:
+            self._keys[key.upper()] = key
+        return ret or key
+
+    def update(self, rhs):
+        for key, value in rhs.items():
+            self[key] = value
+
+    def read_changes(self):
+        for key in self._edits:
+            if key in self._env:
+                yield key, self._env[key]
+
+        for key, value in self._env.items():
+            if value != self._original.get(key, None):
+                yield key, value
+
+        for key in self._original.keys():
+            if key.upper() not in self._keys:
+                yield key, None
+
+    def __iter__(self): return iter(self.keys())
+    def keys(self):     return self._env.keys()
+    def values(self):   return self._env.values()
+    def items(self):    return self._env.items()
+    def get(self, *a):  return self._env.get(*a)
+
+#-------------------------------------------------------------------------------
+class _ExecContext(object):
+    def __init__(self):
+        self._env = None
+
+    def get_env(self):
+        self._env = self._env or _Env()
+        return self._env
+
+    def create_runnable(self, cmd, *args):
+        return _Runnable(cmd, *args, env=self._env)
+
+
+
+#-------------------------------------------------------------------------------
+class _Channel(object):
+    def read_extensions(self, mount):
+        extendable = self.get_channel().get_extendable(mount)
+        yield from (x for x,_ in extendable.read_extensions())
+
+    def get_extension_type(self, mount, name):
+        extendable = self.get_channel().get_extendable(mount)
+        extension = extendable.get_extension_class(name)
+        return extension.get_class_type()
+
+    def read_extension_types(self, mount):
+        extendable = self.get_channel().get_extendable(mount)
+        return extendable.read_extensions()
+
+    def create_extension(self, mount, name, *ctor_args, **ctor_kwargs):
+        extendable = self.get_channel().get_extendable(mount)
+        extension = extendable.get_extension_class(name)
+        return extension.construct(*ctor_args, **ctor_kwargs)
+
+
+
+#-------------------------------------------------------------------------------
+class _Noticeboard(object):
+    def __init__(self, path):
+        self._path = path
+        self._inner = None
+        self._modified = False
+
+    def __del__(self):
+        if self._modified:
+            with open(self._path, "wb") as x:
+                marshal.dump(self._inner, x)
+
+    def _get_inner(self):
+        if not self._inner:
+            try:
+                with open(self._path, "rb") as x:
+                    self._inner = marshal.load(x)
+            except:
+                self._inner = {}
+        return self._inner
+
+    def __getitem__(self, name):
+        return self._get_inner().get(name)
+
+    def __setitem__(self, name, value):
+        self._modified = True
+        self._get_inner()[name] = value
+
+    def copy_from(self, other):
+        self._modified = True
+        self._inner = other._get_inner().copy()
+
+
+
+#-------------------------------------------------------------------------------
+from . import _flick
+Arg = _flick.Arg
+Opt = _flick.Opt
+class Cmd(_flick.Cmd, _Channel):
+    Arg = _flick.Arg
+    Opt = _flick.Opt
+
+    class Noticeboard(Enum):
+        SESSION     = 0
+        PERSISTENT  = 1
+
+    def _print_error(self, message):
+        with text.light_red:
+            super()._print_error(message)
+
+    def _print_help(self, *args, **kwargs):
+        with text.light_cyan:
+            super()._print_help(*args, **kwargs)
+
+    def _call_main(self):
+        # Some host shells such as Bash set $PWD to the current directory. p4
+        # detects this and uses it instead of its current directory. This can
+        # cause subtle bugs in Bash-hosted ushell sessions with patterns like
+        # 'os.chdir(x); subprocess.run("p4 set")'. Unset PWD.
+        if "PWD" in os.environ:
+            del os.environ["PWD"]
+
+        try:
+            return super()._call_main()
+        except KeyboardInterrupt as e:
+            print()
+            self.print_warning("Interrupted!", str(e))
+            raise SystemExit(1)
+
+    def post_construct(self):
+        pass
+
+    def is_interactive(self):
+        if os.name != "nt":
+            return sys.stdout.isatty() and sys.stdin.isatty()
+
+        # Windows says "NUL" is a character device so MSVC's isatty() returns
+        # true. flow.native.isatty() is more accurate.
+        from . import native as flow_native
+        return flow_native.isatty(0) and flow_native.isatty(1) # 0=in, 1=out
+
+    def get_channel(self):
+        return self._channel
+
+    def get_home_dir(self):
+        return os.path.expanduser("~/.ushell/")
+
+    def get_noticeboard(self, board_type):
+        temp_dir = self._channel.get_system().get_temp_dir()
+        if board_type == Cmd.Noticeboard.SESSION:
+            flow_sid = os.getenv("FLOW_SID", "x")
+            return _Noticeboard(temp_dir + "session_" + flow_sid)
+        elif board_type == Cmd.Noticeboard.PERSISTENT:
+            return _Noticeboard(temp_dir + "noticeboard_p")
+
+    def edit_file(self, path):
+        if editor := os.getenv("GIT_EDITOR") or os.getenv("P4EDITOR"):
+            if not os.path.isfile(editor):
+                import shlex
+                exe, *args = (*shlex.split(editor), path)
+            else:
+                exe, *args = (editor, path)
+        else:
+            exe = "notepad.exe" if os.name == "nt" else "nano"
+            args = (path,)
+
+        import subprocess
+        subprocess.run((exe, *args))
+
+    def print_info(self, *message):
+        _lovely_info(*message)
+
+    def print_warning(self, *message):
+        _lovely_warning(*message)
+
+    def print_error(self, *message):
+        _lovely_error(*message)
+
+    def get_exec_context(self):
+        return _ExecContext()
+
+    @staticmethod
+    def summarise(main_method):
+        import time
+        import datetime
+
+        def finish(start_time, result):
+            exec_time = time.time() - start_time
+
+            if isinstance(result, bool):
+                result = 0 if result is True else 1
+
+            print("")
+            if result: _lovely_info("Result: " + text.light_red("Failed - " + hex(result)))
+            else: _lovely_info("Result: " + text.green("Success"))
+            _lovely_info("  Time:", str(datetime.timedelta(seconds=int(exec_time))))
+
+            return result
+
+        def begin(*args, **kwargs):
+            start_time = time.time()
+            result = main_method(*args, **kwargs)
+            if getattr(result, "cr_code", None):
+                async def inner(result):
+                    result = await result
+                    return finish(start_time, result)
+                return inner(result)
+
+            return finish(start_time, result)
+        return begin
diff --git a/Engine/Extras/ushell/channels/flow/core/system/flow/describe.py b/Engine/Extras/ushell/channels/flow/core/system/flow/describe.py
new file mode 100644
--- /dev/null
+++ b/Engine/Extras/ushell/channels/flow/core/system/flow/describe.py
@@ -0,0 +1,82 @@
+# Copyright Epic Games, Inc. All Rights Reserved.
+
+from . import _backtrace
+
+#-------------------------------------------------------------------------------
+class Tool(object):
+    class _Bundle(object):
+        def __init__(self):                    self._payloads = []
+        def payload(self, url, filetype=None): self._payloads.append((url, filetype))
+
+    def __init__(self):
+        self._bin_paths = []
+        self._bundles = {}
+        self._root_dir = None
+        self._sha1 = None
+        self._version = None
+        self._platform = None
+        self._source = (None, None)
+
+    def sha1(self, sha1):                  self._sha1 = sha1
+    def version(self, version):            self._version = version
+    def platform(self, platform):          self._platform = platform
+    def source(self, url, ver_regex):      self._source = (url, ver_regex)
+    def bin(self, bin_path):               self._bin_paths.append(bin_path)
+    def root_dir(self, root_dir):          self._root_dir = root_dir
+    def bundle(self, name="$"):            return self._bundles.setdefault(name, Tool._Bundle())
+    def payload(self, url, filetype=None): self.bundle().payload(url, filetype)
+
+    # Implemented elsewhere.
+    #def extract(self, src_path, dest_dir):
+        #pass
+
+#-------------------------------------------------------------------------------
+class Command(object):
+    def __init__(self):
+        self._path = None
+        self._prefix = "."
+        self._py_path = None
+        self._py_class = None
+
+    def prefix(self, prefix):
+        self._prefix = prefix
+
+    def invoke(self, *path):
+        self._path = path
+
+    def source(self, script_path, cmd_class):
+        self._py_path = script_path
+        self._py_class = cmd_class
+
+#-------------------------------------------------------------------------------
+class Extension(object):
+    def __init__(self):
+        self._name = None
+        self._mount = None
+        self._py_path = None
+        self._py_class = None
+
+    def mount(self, extendable, name):
+        self._extendable = extendable
+        self._name = name
+
+    def source(self, script_path, ext_class):
+        self._py_path = script_path
+        self._py_class = ext_class
+
+#-------------------------------------------------------------------------------
+class Channel(object):
+    def __init__(self):
+        self._parent = ""
+        self._pips = []
+        self._extendables = []
+        self._version = None
+
+    def parent(self, parent):   self._parent = parent
+    def version(self, version): self._version = version
+    def extendable(self, name): self._extendables.append(name)
+
+    def _pip_deprecated_(self, name):
+        # Please do not use! Pips come with security and licensing headaches so
+        # support for them has been removed.
+        self._pips.append(name)
diff --git a/Engine/Extras/ushell/channels/flow/core/system/flow/native/build.bat b/Engine/Extras/ushell/channels/flow/core/system/flow/native/build.bat
new file mode 100755
--- /dev/null
+++ b/Engine/Extras/ushell/channels/flow/core/system/flow/native/build.bat
@@ -0,0 +1,53 @@
+:: Copyright Epic Games, Inc. All Rights Reserved.
+
+@echo off
+pushd %~dp0
+
+cl.exe 1>nul 2>nul
+if errorlevel 9009 (
+    for /f "delims=" %%d in ('vswhere -latest -find VC/Auxiliary/Build/vcvars64.bat') do (
+        call "%%d"
+    )
+)
+
+setlocal
+
+set pysdk=%~f1
+if "%pysdk%"=="" goto:usage
+
+if not "%3"=="release" (
+    set _cl=/Od
+    set _link=/debug
+) else (
+    set _cl=/O1 /Os
+)
+
+set _cl="/I%pysdk%/include" /nologo /Zl /Zi /EHs-c- /c %_cl%
+set _link=%_link% /nologo /dll /machine:x64^
+    /nodefaultlib /align:16^
+    /incremental:no /opt:ref,icf^
+    "/libpath:%pysdk%/libs" python3.lib kernel32.lib
+
+mkdir 1>nul 2>nul _build
+cd _build
+
+call:go cl.exe ../native.cpp %_cl%
+if errorlevel 1 exit /b 1
+
+call:go link.exe %_link% /out:native.pyd native.obj
+if errorlevel 1 exit /b 1
+
+move /y native.pyd ..\__init__.pyd
+
+popd
+goto:eof
+
+:go
+echo.
+echo [96m%*[0m
+%*
+goto:eof
+
+:usage
+echo usage: %~nx0 pysdk [release]
+goto:eof
diff --git a/Engine/Extras/ushell/channels/flow/core/system/flow/native/native.cpp b/Engine/Extras/ushell/channels/flow/core/system/flow/native/native.cpp
new file mode 100644
--- /dev/null
+++ b/Engine/Extras/ushell/channels/flow/core/system/flow/native/native.cpp
@@ -0,0 +1,71 @@
+// Copyright Epic Games, Inc. All Rights Reserved.
+
+#define PY_SSIZE_T_CLEAN
+
+#include <Python.h>
+#include <Windows.h>
+
+//------------------------------------------------------------------------------
+static PyObject* isatty_impl(PyObject* self, PyObject** args, Py_ssize_t nargs)
+{
+    HANDLE Handle = nullptr;
+
+    if (nargs > 0 && PyLong_Check(args[0]))
+    {
+        switch (PyLong_AsLong(args[0]))
+        {
+        case 0: Handle = GetStdHandle(STD_INPUT_HANDLE);    break;
+        case 1: Handle = GetStdHandle(STD_OUTPUT_HANDLE);   break;
+        case 2: Handle = GetStdHandle(STD_ERROR_HANDLE);    break;
+        }
+    }
+
+    if (Handle == nullptr)
+        Py_RETURN_NONE;
+
+    DWORD Type = GetFileType(Handle);
+    if (Type == FILE_TYPE_CHAR)
+    {
+        DWORD Mode;
+        if (GetConsoleMode(Handle, &Mode))
+            Py_RETURN_TRUE;
+    }
+
+    Py_RETURN_FALSE;
+}
+
+//------------------------------------------------------------------------------
+static PyMethodDef native_module_methods[] = {
+    { "isatty", PyCFunction(isatty_impl), METH_FASTCALL, "Like isatty() but correctly reports false for NUL"},
+    {},
+};
+
+//------------------------------------------------------------------------------
+static PyModuleDef native_module = {
+    PyModuleDef_HEAD_INIT,
+    "native",
+    nullptr,
+    -1,
+    native_module_methods,
+};
+
+//------------------------------------------------------------------------------
+extern "C" __declspec(dllexport)
+PyObject* PyInit_native()
+{
+    PyObject* module = PyModule_Create(&native_module);
+    return module;
+}
+
+
+
+//------------------------------------------------------------------------------
+BOOL WINAPI
+_DllMainCRTStartup(
+        HANDLE  hDllHandle,
+        DWORD   dwReason,
+        LPVOID  lpreserved
+        )
+{
+    return TRUE;
+}
diff --git a/Engine/Extras/ushell/channels/flow/core/system/flow/system.py b/Engine/Extras/ushell/channels/flow/core/system/flow/system.py
new file mode 100644
--- /dev/null
+++ b/Engine/Extras/ushell/channels/flow/core/system/flow/system.py
@@ -0,0 +1,260 @@
+# Copyright Epic Games, Inc. All Rights Reserved.
+
+import os
+import sys
+import marshal
+import importlib.util
+
+#-------------------------------------------------------------------------------
+class _FlowObject(object):
+    def __init__(self, system):
+        self._system = system
+
+    def get_system(self):
+        return self._system
+
+
+
+#-------------------------------------------------------------------------------
+class _LoadableClass(_FlowObject):
+    def __init__(self, system, name, specs):
+        super().__init__(system)
+        self._name = name
+        self._specs = specs
+        self._class_type = None
+
+    def get_specs(self):
+        return self._specs
+
+    def _load_spec(self, cmd_spec):
+        channel_id, py_path, py_class = cmd_spec
+
+        channel = self.get_system().get_channel(channel_id)
+        full_py_path = channel.get_path() + py_path
+
+        spec = importlib.util.spec_from_file_location("", full_py_path)
+        module = importlib.util.module_from_spec(spec)
+        spec.loader.exec_module(module)
+
+        ret = getattr(module, py_class, None)
+        if not ret:
+            raise LookupError(f"Failed to create class '{channel}/{py_path}:{py_class}'")
+
+        return ret
+
+    def get_class_type(self):
+        if not self._class_type:
+            specs = tuple(self._load_spec(x) for x in self._specs)
+            self._class_type = type(self._name, specs, {})
+        return self._class_type
+
+    def construct(self, *args, **kwargs):
+        return self.get_class_type()(*args, **kwargs)
+
+
+
+#-------------------------------------------------------------------------------
+class _CommandClass(_LoadableClass):
+    def __init__(self, system, channel_id, specs):
+        super().__init__(system, "flow.cmd", specs)
+        self._channel_id = channel_id
+
+    def construct(self):
+        instance = super().construct()
+        instance._channel = self.get_system().get_channel(self._channel_id)
+        instance.post_construct()
+        return instance
+
+#-------------------------------------------------------------------------------
+class _CommandNode(_FlowObject):
+    _MAGIC = 493
+
+    def __init__(self, system, node):
+        super().__init__(system)
+        self._node = node
+
+    def read_children(self):
+        system = self.get_system()
+        for name, child_node in self._node.items():
+            if isinstance(name, str):
+                yield name, _CommandNode(system, child_node)
+
+    def find_child(self, name):
+        child = self._node.get(name, None)
+        if child:
+            return _CommandNode(self.get_system(), child)
+
+    def get_command_class(self):
+        specs = self._node.get(_CommandNode._MAGIC, None)
+        if specs:
+            return _CommandClass(self.get_system(), specs[0], specs[1:])
+
+    def is_invokeable(self):
+        return self._node.get(_CommandNode._MAGIC, None) != None
+
+#-------------------------------------------------------------------------------
+class _CommandTree(_FlowObject):
+    def __init__(self, system, root_node):
+        super().__init__(system)
+        self._root_node = root_node
+
+    def get_root_node(self):
+        return _CommandNode(self.get_system(), self._root_node)
+
+    def find(self, arg_iter):
+        arg = None
+        node = self._root_node
+        for arg in arg_iter:
+            try:
+                node = node[arg]
+            except KeyError:
+                break
+        else:
+            arg = None
+
+        def read_ret_args():
+            if arg != None: yield arg
+            yield from arg_iter
+
+        return _CommandNode(self.get_system(), node), read_ret_args()
+
+
+
+#-------------------------------------------------------------------------------
+class _ExtensionClass(_LoadableClass):
+    pass
+
+#-------------------------------------------------------------------------------
+class _Extendable(_FlowObject):
+    def __init__(self, system, data):
+        super().__init__(system)
+        self._data = data
+
+    def read_extensions(self):
+        system = self.get_system()
+        yield from ((k, _ExtensionClass(system, k, v)) for k,v in self._data.items())
+
+    def get_extension_class(self, name):
+        specs = self._data.get(name, None)
+        if specs:
+            return _ExtensionClass(self.get_system(), "ext." + name, specs)
+
+        raise ValueError(f"Unknown extendable '{name}'")
+
+
+
+#-------------------------------------------------------------------------------
+class _Tool(_FlowObject):
+    def __init__(self, system, name, data):
+        super().__init__(system)
+        self._name = name
+        self._data = data
+
+    def get_name(self):
+        return self._name
+
+    def get_dir(self):
+        return self.get_system().get_tools_dir() + f"{self._name}-{self._data[0]}/"
+
+    def get_version(self):
+        return self._data[0]
+
+    def get_bin_paths(self):
+        return self._data[1]
+
+
+
+#-------------------------------------------------------------------------------
+class _Channel(_FlowObject):
+    def __init__(self, system, data):
+        super().__init__(system)
+        self._data = data
+
+    def get_name(self):
+        return self._data["name"]
+
+    def get_path(self):
+        return self._data["path"]
+
+    def get_parent(self):
+        parent = self._data["parent"]
+        if parent >= 0:
+            return self.get_system().get_channel(parent)
+
+    def get_extendable(self, name):
+        extendable = self._data["extendables"].get(name, None)
+        if extendable:
+            return _Extendable(self.get_system(), extendable)
+
+        parent = self.get_parent()
+        if parent:
+            return parent.get_extendable(name)
+
+        raise ValueError(f"Unknown extendable '{name}'")
+
+    def read_tools(self):
+        for name, tool in self._data["tools"].items():
+            yield _Tool(self.get_system(), name, tool)
+
+    def get_tool(self, name):
+        tool = self._data["tools"].get(name, None)
+        if tool:
+            return _Tool(self.get_system(), name, tool)
+
+        parent = self.get_parent()
+        if parent:
+            return parent.get_tool(name)
+
+        raise ValueError(f"Unknown tool '{name}'")
+
+
+
+#-------------------------------------------------------------------------------
+class System(object):
+    def __init__(self, working_dir):
+        self._working_dir = working_dir
+        self._manifest = None
+        self._mtime = 0
+
+    def _get_manifest(self):
+        if not self._manifest:
+            with open(self.get_working_dir() + "manifest", "rb") as x:
+                self._manifest = marshal.load(x)
+
+            for pylib_dir in self._manifest["pylib_dirs"]:
+                sys.path.append(pylib_dir)
+        return self._manifest
+
+    def get_working_dir(self):
+        return self._working_dir
+
+    def get_temp_dir(self):
+        ret = self.get_working_dir() + "temp/"
+        try: os.mkdir(ret)
+        except FileExistsError: pass
+        return ret
+
+    def get_tools_dir(self):
+        manifest = self._get_manifest()
+        return manifest["tools_dir"]
+
+    def get_command_tree(self):
+        manifest = self._get_manifest()
+        return _CommandTree(self, manifest["cmd_tree"])
+
+    def read_channels(self):
+        manifest = self._get_manifest()
+        yield from (_Channel(self, x) for x in manifest["channels"])
+
+    def get_channel(self, index):
+        manifest = self._get_manifest()
+        return _Channel(self, manifest["channels"][index])
+
+    def is_path_stale(self, path):
+        manifest_path = self.get_working_dir() + "manifest"
+        try:
+            self._mtime = self._mtime or os.path.getmtime(manifest_path)
+            test_mtime = os.path.getmtime(path)
+            return test_mtime <= self._mtime
+        except FileNotFoundError:
+            return True
diff --git a/Engine/Extras/ushell/channels/flow/core/system/flow/text.py b/Engine/Extras/ushell/channels/flow/core/system/flow/text.py
new file mode 100644
--- /dev/null
+++ b/Engine/Extras/ushell/channels/flow/core/system/flow/text.py
@@ -0,0 +1,132 @@
+# Copyright Epic Games, Inc. All Rights Reserved.
+
+import sys
+
+_enabled = sys.stdout.isatty() and sys.stderr.isatty()
+
+#-------------------------------------------------------------------------------
+class _Deferred(object):
+    def __init__(self, builder, string):
+        def _deferred_pre(): return builder._get_tag() or ""
+        def _deferred_post(): return builder._get_prev_tag() or ""
+        self.actions = [_deferred_pre, string, _deferred_post]
+
+    def __mod__(self, rhs):
+        if not isinstance(rhs, tuple): rhs = rhs,
+        self.actions.append(rhs)
+        return self
+
+    def __add__(self, rhs):
+        self.actions.append(rhs)
+        return self
+
+    def __radd__(self, rhs):
+        self.actions.insert(0, rhs)
+        return self
+
+    def __str__(self):
+        out = ""
+        for action in self.actions:
+            if hasattr(action, "__call__"): out += action()
+            elif isinstance(action, tuple): out = out % action
+            else: out += str(action)
+        return out
+
+#-------------------------------------------------------------------------------
+class _Builder(object):
+    _prev = None
+    _colours = {
+        "_"             : None,
+        "black"         : 0,
+        "red"           : 1,
+        "green"         : 2,
+        "yellow"        : 3,
+        "blue"          : 4,
+        "magenta"       : 5,
+        "cyan"          : 6,
+        "light_grey"    : 7,
+        "grey"          : 8,
+        "light_red"     : 9,
+        "light_green"   : 10,
+        "light_yellow"  : 11,
+        "light_blue"    : 12,
+        "light_magenta" : 13,
+        "light_cyan"    : 14,
+        "white"         : 15,
+    }
+
+    def __init__(self, fg=None):
+        self.prev = None
+        self.bg = None
+        try:
+            self.fg = self._colours[fg]
+        except KeyError:
+            raise SyntaxError("Invalid foreground colour '%s'" % fg)
+
+    def __getattr__(self, name):
+        try:
+            self.bg = self._colours[name]
+            return self
+        except KeyError:
+            pass
+
+        try:
+            return self.__dict__[name]
+        except KeyError:
+            raise SyntaxError("Invalid background colour '%s'" % name)
+
+    def __call__(self, string):
+        return _Deferred(self, string)
+
+    def __enter__(self):
+        self.prev = _Builder._prev
+        _Builder._prev = self
+
+        tag = self._get_tag()
+        if tag:
+            sys.stdout.write(tag)
+
+    def __exit__(self, *args):
+        _Builder._prev = self.prev
+
+        if self.fg or self.bg:
+            tag = self._get_prev_tag()
+            if tag:
+                sys.stdout.write(tag)
+
+    def _get_tag(self):
+        return self._get_tag_impl(self.fg, self.bg)
+
+    def _get_prev_tag(self):
+        prev = self.prev or _Builder._prev
+        if prev:
+            return prev._get_tag()
+        elif _enabled:
+            return "\x1b[0m"
+
+    @staticmethod
+    def _get_tag_impl(fg=None, bg=None):
+        if not _enabled:
+            return
+
+        if fg == None and bg == None:
+            return "\x1b[0m"
+
+        codes = []
+        if fg != None and fg >= 0:
+            bias = 90 if fg > 7 else 30
+            codes.append(bias + (fg & 7))
+
+        if bg != None and bg >= 0:
+            bias = 100 if bg > 7 else 40
+            codes.append(bias + (bg & 7))
+
+        if len(codes):
+            return "\x1b[" + ";".join((str(x) for x in codes)) + "m"
+
+#-------------------------------------------------------------------------------
+class _Text(object):
+    def __getattr__(self, attr_name):
+        return _Builder(attr_name)
+
+text = _Text()
