Updated ushell [26875160]

diff --git a/Engine/Extras/ushell/channels/flow/core/system/lib/bootstrap.py b/Engine/Extras/ushell/channels/flow/core/system/lib/bootstrap.py
new file mode 100644
--- /dev/null
+++ b/Engine/Extras/ushell/channels/flow/core/system/lib/bootstrap.py
@@ -0,0 +1,644 @@
+# Copyright Epic Games, Inc. All Rights Reserved.
+
+import os
+import sys
+import fsutils
+import hashlib
+import marshal
+import flow.describe
+import subprocess as sp
+
+#-------------------------------------------------------------------------------
+class _Log(object):
+    def __init__(self, log_path):
+        self._indent = 0
+        try: self._log_file = open(log_path, "wt")
+        except: self._log_file = None
+
+    def __del__(self):
+        if self._log_file:
+            self._log_file.close()
+
+    def indent(self, *name):
+        self.write("Section:", *name)
+        self._indent += 1
+
+    def unindent(self):
+        self._indent -= 1
+
+    def write(self, *args):
+        if self._log_file:
+            print(*args, file=self._log_file)
+
+    def print(self, *args, **kwargs):
+        print("  " * (self._indent * (1 if kwargs.get("end") == None else 0)), end="")
+        print(*args, **kwargs, file=sys.stdout)
+
+    def print_write(self, *args, **kwargs):
+        self.print(*args, **kwargs)
+        self.write(*args)
+
+_log = None
+
+
+
+#-------------------------------------------------------------------------------
+def _http_get(url, dest_dir, progress_cb=None):
+    import re
+    from urllib.request import urlopen
+
+    # Fire up an http client to get the url.
+    client = urlopen(url)
+    if (client.status // 100) != 2:
+        assert False, f"Error creating HTTP client for {url}"
+
+    if progress_cb:
+        progress_cb("", 0.0)
+
+    # Get the download's file name
+    content_header = client.headers.get("content-disposition", "")
+    m = re.search(r'filename="?([^;"]+)"?(;|$)', content_header)
+    if m: file_name = m.group(1)
+    else: file_name = os.path.basename(client.url)
+    dest_path = dest_dir + file_name
+    assert not os.path.exists(dest_path), f"Directory '{dest_path}' unexpectedly exists"
+
+    # Do the download. If there's no content-length we are probably receiving
+    # chunked transfer-encoding. So we'll cheat.
+    chunk_size = 128 << 10
+    if chunk_count := client.headers.get("content-length"):
+        chunk_count = int(chunk_count)
+    else:
+        chunk_count = 1 << 48;
+    chunk_count = (chunk_count + chunk_size - 1) // chunk_size
+    chunk_index = 0
+    with open(dest_path, "wb") as out:
+        while chunk := client.read(chunk_size):
+            out.write(chunk)
+            out.flush()
+            chunk_index += 1
+            if progress_cb:
+                progress_cb(file_name, chunk_index / chunk_count)
+
+    return dest_path
+
+
+
+#-------------------------------------------------------------------------------
+def _extract_zip(payload, dest_dir):
+    import zipfile
+    zip_file = zipfile.ZipFile(payload, "r")
+    zip_file.extractall(dest_dir)
+    zip_file.close()
+
+#-------------------------------------------------------------------------------
+def _extract_tar(payload, dest_dir):
+    cmd = ("tar", "-x", "-C", dest_dir, "-f", payload)
+    ret = sp.run(cmd, stdout=sp.DEVNULL)
+    assert ret.returncode == 0, "Extract failed; " + " ".join(cmd)
+
+#-------------------------------------------------------------------------------
+def _tool_extract_method(self, src_path, dest_dir):
+    extractors = {
+        ".tar.gz"   : _extract_tar,
+        ".tgz"      : _extract_tar,
+        ".zip"      : _extract_zip,
+    }
+
+    was_archive = False
+    for file_name in (x.lower() for x in os.listdir(src_path)):
+        for suffix, extract_func in extractors.items():
+            if not file_name.endswith(suffix):
+                continue
+
+            _log.write("extracting", file_name, "to", dest_dir, "with", extract_func)
+            extract_func(src_path + file_name, dest_dir)
+            was_archive = True
+            break
+
+    if was_archive:
+        return
+
+    for file_name in os.listdir(src_path):
+        os.rename(src_path + file_name, dest_dir + file_name)
+
+flow.describe.Tool.extract = _tool_extract_method
+
+
+
+#-------------------------------------------------------------------------------
+def _validate_tool(name, tool):
+    assert tool._version, f"No version set on tool '{name}'"
+    assert tool._bundles, f"No payload set on tool '{name}'"
+
+#-------------------------------------------------------------------------------
+def _acquire_tool(name, tool, target_dir, progress_cb):
+    # Some boiler plate to provide a progress update
+    payload_i = 0
+    payload_n = sum(len(x._payloads) for x in tool._bundles.values())
+    def inner_progress_cb(file_name, value):
+        value = int(((float(payload_i) + value) * 100) / payload_n)
+        progress_cb(f"{value:3}%")
+
+    # Acquire the tool's bundles/payloads
+    temp_dir = fsutils.WorkPath(target_dir[:-1] + ".work/")
+    for bundle_name, bundle in tool._bundles.items():
+        _log.write("Bundle", bundle_name)
+        fetch_dir = temp_dir + bundle_name + "/"
+        os.mkdir(fetch_dir)
+        for payload in bundle._payloads:
+            _log.write("Payload", *payload)
+            payload, ftype = payload
+            payload_dest = _http_get(payload, fetch_dir, inner_progress_cb)
+            payload_i += 1
+            if ftype:
+                os.rename(payload_dest, payload_dest + "." + ftype)
+
+    # Hash the downloaded content to confirm it is what is expected
+    progress_cb("sha1")
+    sha1 = hashlib.sha1()
+    for bundle_name, bundle in tool._bundles.items():
+        for item in fsutils.read_files(temp_dir + bundle_name):
+            with open(item.path, "rb") as in_file:
+                while True:
+                    data = in_file.read(16384)
+                    if len(data) <= 0:
+                        break
+                    sha1.update(data)
+    tool_sha1 = sha1.hexdigest()
+    _log.write("SHA1", name, tool_sha1)
+
+    if tool._sha1:
+        expected_sha1 = tool._sha1.lower()
+        if tool_sha1 != expected_sha1:
+            assert false, f"Unexpected content bundle data for tool '{name}' [sha1:{tool_sha1}]"
+
+    # Extract the payloads.
+    progress_cb("uzip")
+    extract_dir = temp_dir + "$extract/"
+    os.mkdir(extract_dir)
+    for bundle_name, bundle in tool._bundles.items():
+        _log.write("Extract", bundle_name, "from", fetch_dir, "to", extract_dir, "with", tool.extract)
+        fetch_dir = temp_dir + bundle_name + "/"
+        tool.extract(fetch_dir, extract_dir)
+
+    # Collapse
+    source_dir = str(extract_dir)
+    while True:
+        dir_iter = iter(fsutils.read_files(source_dir))
+        x = next(dir_iter, None)
+        if x and not next(dir_iter, None) and x.is_dir():
+            source_dir = x.path
+            continue
+        break
+
+    if tool._root_dir:
+        source_dir += "/" + tool._root_dir
+        assert os.path.isdir(source_dir), "Given root f'{tool._root_dir}' for tool '{name}' is not valid"
+
+    # Store sha1
+    with open(source_dir + "/sha1.flow.txt", "w") as out:
+        out.write(tool_sha1)
+
+    os.rename(source_dir, target_dir)
+
+    progress_cb("done")
+    return tool_sha1
+
+#-------------------------------------------------------------------------------
+def _manifest_tool(name, tool):
+    bundles = {}
+    for bundle_name, bundle in tool._bundles.items():
+        bundles[bundle_name] = tuple(x[0] for x in bundle._payloads)
+
+    return {
+        "sha1"      : 0,
+        "version"   : tool._version,
+        "double"    : f"{name}-{tool._version}",
+        "bin_paths" : tuple(tool._bin_paths),
+        "bundles"   : bundles,
+        "source"    : tool._source,
+    }
+
+#-------------------------------------------------------------------------------
+def _install_tool(name, tool, channel_dir):
+    manifest = _manifest_tool(name, tool)
+    if tool._platform and tool._platform != sys.platform:
+        return manifest
+
+    tool_double = f"{name}-{tool._version}/"
+    _log.write("Installing tool", tool_double)
+
+    log_line = (13 * " ") + tool_double[:-1]
+    log_line += "\b" * (len(log_line) - 1)
+    _log.print(log_line, end="\r      [    ]\b")
+
+    # If the tool doesn't exist then it needs acquiring
+    tool_dir = "../tools/" + tool_double
+    if not os.path.isdir(tool_dir):
+        def progress_cb(info):
+            _log.print("\b" * 4, info[:4], sep="", end="")
+
+        try:
+            _log.write("Acquiring to", tool_dir)
+            sha1 = _acquire_tool(name, tool, tool_dir, progress_cb)
+        except KeyboardInterrupt:
+            raise
+        except Exception as e:
+            _log.write("ERROR", str(e))
+            _log.print("\b" * 4, "fail", sep="", end="\n")
+            return manifest
+
+        # Post install
+        if hasattr(tool, "post_install"):
+            progress_cb("post")
+            _log.write("Post-install")
+            tool.post_install(tool_dir)
+
+        # Validate that the required binaries exist.
+        for bin_path in tool._bin_paths:
+            if not os.path.exists(tool_dir + bin_path):
+                _log.write(f"Unable to find '{bin_path}' in '{tool_dir}'")
+    else:
+        sha1 = 1
+        if os.path.isfile(tool_dir + "/sha1.flow.txt"):
+            with open(tool_dir + "/sha1.flow.txt", "rt") as sha1_file:
+                sha1 = sha1_file.read()
+        _log.write("Tool already acquired", tool_dir)
+
+    manifest["sha1"] = sha1
+    _log.print("\b" * 4, " ok ", sep="", end="\n")
+    return manifest
+
+
+
+#-------------------------------------------------------------------------------
+class _Action(object):
+    def __init__(self, name, input_path):
+        self._name = name
+        self._input_path = input_path
+
+    def needs_update(self):
+        return True
+
+
+
+#-------------------------------------------------------------------------------
+def _validate_command(channel_src_dir, name, command):
+    assert command._path, f"Command '{name}' missing call to invoke/composite()"
+    assert command._py_path, f"Command '{name}' missing call to source()"
+    if not os.path.isfile(channel_src_dir + command._py_path):
+        assert False, f"Command '{name}' missing source '{command._py_path}'"
+
+#-------------------------------------------------------------------------------
+def _validate_extension(channel_src_dir, name, extension):
+    assert extension._name, f"Extension '{name}' invalid/missing call to mount()"
+    assert extension._extendable, f"Extension '{name}' missing call to mount()"
+    assert extension._py_path, f"Extension '{name}' missing call to source()"
+    if not os.path.isfile(channel_src_dir + extension._py_path):
+        assert False, f"Extension '{name}' missing source '{extension._py_path}'"
+
+#-------------------------------------------------------------------------------
+def _select_update_type():
+    version_path = os.path.abspath(__file__ + "/../../version")
+    input_deps = (
+        version_path,             # primary version for schema/state invalidation
+        "../python/current/",     # embedded version of Python.
+    )
+    mtime = max(fsutils.get_mtime(x) for x in input_deps)
+    return _UpdateMaybe if fsutils.get_mtime("manifest") > mtime else _UpdateImpl
+
+#-------------------------------------------------------------------------------
+class _UpdateImpl(_Action):
+    def run(self):
+        _log.print(f"Updating channel '{self._name}'")
+        _log.indent("Channel", self._name)
+
+        # Clean up the output. Updates are an install from scratch.
+        out_dir = f"channels/{self._name}/"
+        fsutils.rmdir(out_dir)
+        out_dir = fsutils.WorkPath(out_dir)
+
+        # Load the describe.flow.py as a module.
+        channel_py = self._input_path + "describe.flow.py"
+        module = fsutils.import_script(channel_py)
+
+        def read_items_by_type(type):
+            for name, value in module.__dict__.items():
+                if isinstance(value, type):
+                    yield name, value
+
+        # Get the channel's description.
+        name, channel = next(read_items_by_type(flow.describe.Channel), (None, None))
+        assert channel, f"Channel() instance missing from channel '{self._name}'"
+        assert channel._version != None, f"Channel '{self._name}' has no version"
+
+        # Install channel's pips.
+        if channel._pips:
+            _log.print("Pips:")
+            _log.indent("Pips")
+            for pip_name in channel._pips:
+                _log.print_write(pip_name)
+                cmd = (sys.executable, "-Xutf8", "-Esum", "pip", "install", pip_name)
+                result = sp.run(cmd, stdout=sp.DEVNULL, stderr=sp.DEVNULL)
+                if result.returncode:
+                    _log.print(" ...failed");
+            _log.unindent()
+
+        # Tools
+        tools_manifest = {}
+        tools = [x for x in read_items_by_type(flow.describe.Tool)]
+        if tools:
+            _log.print("Tools:")
+            _log.indent("Tools")
+            for name, tool in tools:
+                _validate_tool(name, tool)
+
+            for name, tool in tools:
+                tools_manifest[name] = _install_tool(name, tool, out_dir)
+            _log.unindent()
+
+        # Commands
+        commands = [x for x in read_items_by_type(flow.describe.Command)]
+        for name, command in commands:
+            _validate_command(self._input_path, name, command)
+            command._path = command._path or (name)
+            command._path = (command._prefix, *command._path)
+
+        # Extensions
+        extensions = [x for x in read_items_by_type(flow.describe.Extension)]
+
+        for name, extension in extensions:
+            _validate_extension(self._input_path, name, extension)
+
+        # Build manifest
+        manifest = {}
+        manifest["name"] = self._name
+        manifest["path"] = self._input_path
+        manifest["parent"] = channel._parent or "flow.core"
+        manifest["extendables"] = tuple(channel._extendables)
+        manifest["tools"] = tools_manifest
+        manifest["commands"] = tuple(v.__dict__ for k,v in commands)
+        manifest["extensions"] = tuple(v.__dict__ for k,v in extensions)
+
+        # Success. Write the action's output
+        with open(out_dir + "manifest", "wb") as x:
+            marshal.dump(manifest, x)
+
+        out_dir.keep()
+        _log.print("Done\n")
+        _log.unindent()
+
+#-------------------------------------------------------------------------------
+class _UpdateMaybe(_UpdateImpl):
+    def needs_update(self):
+        in_mtime = fsutils.get_mtime(self._input_path + "describe.flow.py")
+        out_mtime = fsutils.get_mtime(f"channels/{self._name}/manifest")
+        return out_mtime <= in_mtime
+
+
+
+#-------------------------------------------------------------------------------
+def _plant_cmd_tree(channels):
+    # Build the command tree
+    cmd_tree = {}
+    for channel in channels:
+        for cmd in channel["commands"]:
+            node = cmd_tree
+            prefix, *cmd_path = cmd["_path"]
+            for piece in (prefix + cmd_path[0], *cmd_path[1:]):
+                node = node.setdefault(piece, {})
+
+            cmd_desc = (
+                channel["index"],
+                cmd["_py_path"],
+                cmd["_py_class"],
+            )
+
+            leaf = node.setdefault(493, [channel["index"]])
+            leaf.insert(1, cmd_desc)
+
+    return cmd_tree
+
+#-------------------------------------------------------------------------------
+def _mount_extensions(channel):
+    py_dir = channel["name"] + "/"
+    for extension in channel["extensions"]:
+        extendable = extension["_extendable"]
+        host = channel
+        while True:
+            if extendable in host["extendables"]:
+                ext_desc = (
+                    channel["index"],
+                    extension["_py_path"],
+                    extension["_py_class"],
+                )
+                out = host["extendables"][extendable]
+                out = out.setdefault(extension["_name"], [])
+                out.append(ext_desc)
+                break
+
+            host = host["parent"]
+            if not host:
+                extension_name = extension["_name"]
+                assert False, f"Unable to find extendable '{extendable}' for '{py_dir}{extension_name}'"
+
+#-------------------------------------------------------------------------------
+def _carve_channels(channels):
+    out = [None] * len(channels)
+    for channel in channels:
+        index = channel["index"]
+        parent_index = channel["parent"]["index"] if channel["parent"] else -1
+        tools = {k:(v["version"], v["bin_paths"]) for k,v in channel["tools"].items() if v["sha1"]}
+        out[channel["index"]] = {
+            "name" : channel["name"],
+            "path" : channel["path"],
+            "parent" : parent_index,
+            "extendables" : channel["extendables"],
+            "tools" : tools,
+        }
+
+    return out
+
+#-------------------------------------------------------------------------------
+def _create_shims(manifest):
+    shims_dir = "shims/"
+    _log.write("Shims dir;", shims_dir)
+
+    import shims
+    builder = shims.Builder(shims_dir)
+    builder.clean()
+
+    # Create command shims
+    run_py = os.path.abspath(__file__ + "/../../run.py")
+    manifest_path = os.path.abspath("manifest")
+    args_prefix = ("-Xutf8", "-Esu", run_py, manifest_path)
+    for name in (x for x in manifest["cmd_tree"] if isinstance(x, str)):
+        builder.create_python_shim(name, *args_prefix, name)
+
+    # Create tool shims
+    tool_root = os.path.abspath("../tools") + "/"
+    for channel in reversed(manifest["channels"]):
+        for name, (version, bin_paths) in channel["tools"].items():
+            tool_dir = tool_root + f"{name}-{version}/"
+            for bin_path in bin_paths:
+                name = os.path.basename(bin_path)
+                tool_path = tool_dir + bin_path
+                if os.path.isfile(tool_path):
+                    builder.create_shim(name, tool_path)
+                else:
+                    _log.write("Binary not found;", tool_path)
+
+    builder.commit()
+
+#-------------------------------------------------------------------------------
+class _Finalise(_Action):
+    def needs_update(self):
+        self._manifest_paths = [x.path + "/manifest" for x in fsutils.read_dirs(self._input_path)]
+        mtime = max(fsutils.get_mtime(x) for x in self._manifest_paths)
+        mtime = max(mtime, fsutils.get_mtime(self._input_path))
+        return fsutils.get_mtime("manifest") <= mtime
+
+    def run(self):
+        # Load channel manifests
+        manifests = []
+        for manifest_path in self._manifest_paths:
+            with open(manifest_path, "rb") as x:
+                manifest = marshal.load(x)
+                manifests.append(manifest)
+
+        # Assign each manifest an index
+        for i in range(len(manifests)):
+            manifests[i]["index"] = i
+
+        # Connect children to their parents
+        manifests = {x["name"]:x for x in manifests}
+        for manifest in manifests.values():
+            parent = manifest["parent"]
+            assert parent in manifests, f"Invalid parent '{parent}' for channel '{manifest['name']}'"
+            manifest["parent"] = manifests[parent] if parent != manifest["name"] else None
+
+        # Sort the manifests topologically
+        topo_manifests = {}
+
+        def topo_insert(manifest):
+            if manifest:
+                topo_insert(manifest["parent"])
+                topo_manifests[manifest["name"]] = manifest
+
+        for manifest in manifests.values():
+            topo_insert(manifest)
+
+        manifests = topo_manifests
+        _log.write("Channels;", *(x["name"] for x in manifests.values()))
+
+        # Prepare extendables and mount extensions
+        for manifest in manifests.values():
+            manifest["extendables"] = {x:{} for x in manifest["extendables"]}
+            _mount_extensions(manifest)
+
+        # Collect various paths for each channel
+        pylib_dirs = []
+        for manifest in manifests.values():
+            pylib_dir = manifest["path"] + "pylib/"
+            if os.path.isdir(pylib_dir):
+                pylib_dirs.append(pylib_dir)
+        _log.write("Pylib dirs;", *(x for x in pylib_dirs))
+
+        # Build the primary manifest.
+        primary = {
+            "channels"   : _carve_channels(manifests.values()),
+            "pylib_dirs" : tuple(pylib_dirs),
+            "tools_dir"  : os.path.abspath("../tools") + "/",
+            "cmd_tree"   : _plant_cmd_tree(manifests.values()),
+        }
+
+        # Create shims
+        _create_shims(primary)
+
+        # Write the primary manifest out.
+        with open("manifest", "wb") as x:
+            marshal.dump(primary, x)
+
+
+
+#-------------------------------------------------------------------------------
+class _Prune(_Action):
+    def run(self):
+        _log.print_write(f"Removing old channel '{self._name}'")
+        fsutils.rmdir(self._input_path)
+
+
+
+#-------------------------------------------------------------------------------
+def impl(working_dir, *channels_dirs):
+    channels_dirs = [os.path.abspath(x) for x in channels_dirs if x]
+    working_dir = os.path.abspath(working_dir)
+
+    def _read_channels(dir, depth=0):
+        for item in fsutils.read_dirs(dir):
+            if os.path.isfile(item.path + "/describe.flow.py"):
+                yield item.name, os.path.abspath(item.path) + "/"
+            elif depth == 0:
+                for name, path in _read_channels(item.path, depth + 1):
+                    yield f"{item.name}.{name}", path
+
+    # Seed the working dir with a hash of our location so we can detect movement
+    loc_hash = 5381
+    for c in os.path.abspath(__file__):
+        loc_hash = ((loc_hash << 5) + loc_hash) + (ord(c) | 0x20)
+        loc_hash &= 0x7fffffff
+    working_suffix = "flow_%08x" % loc_hash
+
+    # Change to the main working dir
+    working_dir += f"/{working_suffix}/"
+    try: os.makedirs(working_dir)
+    except FileExistsError: pass
+    boot_dir = os.getcwd()
+    os.chdir(working_dir)
+
+    lock = fsutils.DirLock("..")
+
+    # Build a list of actions to do. The flow channel must be first
+    actions = {"flow.core" : None}
+
+    # First; add existing channels for removal
+    for dest_dir in fsutils.read_dirs("channels"):
+        actions[dest_dir.name] = _Prune(dest_dir.name, dest_dir.path + "/")
+
+    # Second; add channels from source to get updated.
+    update_type = _select_update_type()
+    for channels_dir in channels_dirs:
+        for name, path in _read_channels(channels_dir):
+            actions[name] = update_type(name, path)
+
+    # Lastly; aggregate channels into a manifest and create command shims
+    actions["*finalise"] = _Finalise("*finalise", "channels/")
+
+    # If there's something to do, do it or wait for it to be done
+    for action in (x for x in actions.values() if x.needs_update()):
+        if lock.claim():
+            global _log
+            if not _log:
+                _log = _Log(working_dir + "setup.log")
+                _log.write("Working dir;", working_dir)
+
+            try:
+                _log.write("-- Running action;", type(action).__name__)
+                action.run()
+            except:
+                import traceback
+                _log.write(traceback.format_exc())
+                raise
+            continue
+
+        import time
+        print("Waiting for another instance to provision.", end="")
+        while not lock.can_claim():
+            time.sleep(1)
+            print(".", end="")
+
+        break
+
+    os.chdir(boot_dir)
+    return working_suffix
diff --git a/Engine/Extras/ushell/channels/flow/core/system/lib/exeshim/__init__.py b/Engine/Extras/ushell/channels/flow/core/system/lib/exeshim/__init__.py
new file mode 100644
--- /dev/null
+++ b/Engine/Extras/ushell/channels/flow/core/system/lib/exeshim/__init__.py
@@ -0,0 +1,193 @@
+# Copyright Epic Games, Inc. All Rights Reserved.
+
+import os
+import sys
+import struct
+
+#-------------------------------------------------------------------------------
+class _K(object):
+    py_section = 568
+    py_subsystem = 276
+    py_alignment = 16
+    py_data = 2720
+    py_payload = (
+        b"\x4d\x5a\x90\x00\x03\x00\x00\x00\x04\x00\x00\x00\xff\xff\x00\x00\xb8\x00\x00\x00\x00\x00\x00\x00\x40\x00\x00\x00\x00\x00\x00\x00\x00\x00\x00\x00\x00\x00\x00\x00\x00\x00\x00\x00\x00\x00\x00\x00\x00\x00\x00\x00\x00\x00\x00\x00\x00\x00\x00\x00\xb8\x00\x00\x00"
+        b"\x0e\x1f\xba\x0e\x00\xb4\x09\xcd\x21\xb8\x01\x4c\xcd\x21\x54\x68\x69\x73\x20\x70\x72\x6f\x67\x72\x61\x6d\x20\x63\x61\x6e\x6e\x6f\x74\x20\x62\x65\x20\x72\x75\x6e\x20\x69\x6e\x20\x44\x4f\x53\x20\x6d\x6f\x64\x65\x2e\x0d\x0d\x0a\x24\x00\x00\x00\x00\x00\x00\x00"
+        b"\x39\x88\x04\xc5\x7d\xe9\x6a\x96\x7d\xe9\x6a\x96\x7d\xe9\x6a\x96\x36\x91\x6b\x97\x7e\xe9\x6a\x96\x7d\xe9\x6b\x96\x74\xe9\x6a\x96\xbc\x95\x6f\x97\x7c\xe9\x6a\x96\xbc\x95\x68\x97\x7c\xe9\x6a\x96\x52\x69\x63\x68\x7d\xe9\x6a\x96\x50\x45\x00\x00\x64\x86\x04\x00"
+        b"\xe3\xc2\x11\x64\x00\x00\x00\x00\x00\x00\x00\x00\xf0\x00\x22\x00\x0b\x02\x0e\x22\xb0\x04\x00\x00\xa0\x03\x00\x00\x00\x00\x00\x00\x50\x05\x00\x00\x60\x02\x00\x00\x00\x00\x00\x40\x01\x00\x00\x00\x10\x00\x00\x00\x10\x00\x00\x00\x06\x00\x00\x00\x00\x00\x00\x00"
+        b"\x06\x00\x00\x00\x00\x00\x00\x00\x00\x40\x00\x00\x60\x02\x00\x00\x00\x00\x00\x00\x03\x00\x60\x81\x00\x00\x10\x00\x00\x00\x00\x00\x00\x10\x00\x00\x00\x00\x00\x00\x00\x00\x10\x00\x00\x00\x00\x00\x00\x10\x00\x00\x00\x00\x00\x00\x00\x00\x00\x00\x10\x00\x00\x00"
+        b"\x00\x00\x00\x00\x00\x00\x00\x00\x20\x09\x00\x00\x28\x00\x00\x00\x00\x00\x00\x00\x00\x00\x00\x00\x50\x0a\x00\x00\x48\x00\x00\x00\x00\x00\x00\x00\x00\x00\x00\x00\x00\x00\x00\x00\x00\x00\x00\x00\x9c\x07\x00\x00\x1c\x00\x00\x00\x00\x00\x00\x00\x00\x00\x00\x00"
+        b"\x00\x00\x00\x00\x00\x00\x00\x00\x00\x00\x00\x00\x00\x00\x00\x00\x00\x00\x00\x00\x00\x00\x00\x00\x00\x00\x00\x00\x00\x00\x00\x00\x10\x07\x00\x00\x50\x00\x00\x00\x00\x00\x00\x00\x00\x00\x00\x00\x00\x00\x00\x00\x00\x00\x00\x00\x00\x00\x00\x00\x00\x00\x00\x00"
+        b"\x2e\x74\x65\x78\x74\x00\x00\x00\xaf\x04\x00\x00\x60\x02\x00\x00\xb0\x04\x00\x00\x60\x02\x00\x00\x00\x00\x00\x00\x00\x00\x00\x00\x00\x00\x00\x00\x20\x00\x00\x60\x2e\x72\x64\x61\x74\x61\x00\x00\x32\x03\x00\x00\x10\x07\x00\x00\x40\x03\x00\x00\x10\x07\x00\x00"
+        b"\x00\x00\x00\x00\x00\x00\x00\x00\x00\x00\x00\x00\x40\x00\x00\x40\x2e\x70\x64\x61\x74\x61\x00\x00\x48\x00\x00\x00\x50\x0a\x00\x00\x50\x00\x00\x00\x50\x0a\x00\x00\x00\x00\x00\x00\x00\x00\x00\x00\x00\x00\x00\x00\x40\x00\x00\x40\x2e\x73\x68\x69\x6d\x00\x00\x00"
+        b"\x04\x00\x00\x00\xa0\x0a\x00\x00\x10\x00\x00\x00\xa0\x0a\x00\x00\x00\x00\x00\x00\x00\x00\x00\x00\x00\x00\x00\x00\x40\x00\x00\xc0\x48\x8b\xc4\x48\x89\x58\x08\x48\x89\x68\x10\x48\x89\x70\x18\x48\x89\x78\x20\x41\x56\x0f\xb7\x11\x33\xff\x4c\x8b\xd9\x4c\x8b\xc9"
+        b"\x4c\x8b\xc1\x44\x8b\xd7\x8d\x77\x20\x85\xd2\x0f\x84\xb3\x00\x00\x00\x0f\xb7\xc2\x8d\x6f\x22\x44\x8d\x77\x27\x66\x41\x89\x01\x49\x8d\x58\x02\x8b\xcf\x45\x85\xd2\x74\x38\x0f\xb7\x03\x41\x3b\xc2\x8b\xc7\x0f\x94\xc1\x83\xfa\x5c\x0f\x94\xc0\x85\xc8\x74\x09\x4c"
+        b"\x8b\xc3\x66\x45\x89\x11\xeb\x65\x41\x3b\xd2\x44\x0f\x44\xd7\x41\x8b\xc2\xf7\xd8\x48\x1b\xc9\x48\xf7\xd1\x83\xe1\x02\x4c\x2b\xc9\xeb\x4b\x66\x44\x39\x33\x8b\xc7\x0f\x94\xc1\x66\x39\x2b\x0f\x94\xc0\x0b\xc8\x8b\xc7\x83\xfa\x5c\x0f\x44\xc1\x85\xc0\x74\x09\x4c"
+        b"\x8b\xc3\x66\x41\x89\x29\xeb\x25\x41\x3b\xd6\x8b\xcf\x8b\xc7\x0f\x94\xc1\x3b\xd5\x0f\x94\xc0\x0b\xc8\xf7\xd9\x45\x1b\xd2\x44\x23\xd2\x74\x06\x49\x83\xe9\x02\xeb\x04\x3b\xd6\x76\x17\x49\x83\xc0\x02\x49\x83\xc1\x02\x41\x0f\xb7\x10\x0f\xb7\xc2\x85\xd2\x0f\x85"
+        b"\x57\xff\xff\xff\x66\x41\x39\x30\x8b\xcf\x8b\xc7\x0f\x96\xc1\x66\x41\x39\x38\x0f\x95\xc0\x85\xc8\x74\x1a\x49\x83\xc0\x02\x8b\xd7\x8b\xcf\x66\x41\x39\x30\x0f\x96\xc2\x66\x41\x39\x38\x0f\x95\xc1\x85\xd1\x75\xe6\x48\x8b\x5c\x24\x10\x4d\x2b\xc3\x48\x8b\x6c\x24"
+        b"\x18\x48\x8b\x74\x24\x20\x49\xd1\xf8\x66\x41\x89\x39\x41\x8b\xc0\x48\x8b\x7c\x24\x28\x41\x5e\xc3\x48\x8b\xc4\x48\x89\x58\x08\x48\x89\x68\x10\x48\x89\x70\x18\x48\x89\x78\x20\x44\x0f\xb7\x09\x33\xff\x48\x8b\xd9\x4c\x8b\xd1\x4c\x8b\xc1\x8b\xd7\x8d\x77\x20\x45"
+        b"\x85\xc9\x0f\x84\x9a\x00\x00\x00\x41\x0f\xb7\xc1\x8d\x6f\x22\x66\x41\x89\x02\x4d\x8d\x58\x02\x8b\xcf\x85\xd2\x74\x37\x41\x0f\xb7\x03\x3b\xc2\x8b\xc7\x0f\x94\xc1\x41\x83\xf9\x5c\x0f\x94\xc0\x85\xc8\x74\x09\x4d\x8b\xc3\x66\x41\x89\x12\xeb\x4d\x44\x3b\xca\x0f"
+        b"\x44\xd7\x8b\xc2\xf7\xd8\x48\x1b\xc9\x48\xf7\xd1\x83\xe1\x02\x4c\x2b\xd1\xeb\x35\x66\x41\x39\x2b\x8b\xc7\x0f\x94\xc1\x41\x83\xf9\x5c\x0f\x94\xc0\x85\xc8\x74\x09\x4d\x8b\xc3\x66\x41\x89\x2a\xeb\x18\x44\x3b\xcd\x8b\xd7\x41\x0f\x44\xd1\x85\xd2\x74\x06\x49\x83"
+        b"\xea\x02\xeb\x05\x44\x3b\xce\x76\x19\x49\x83\xc0\x02\x49\x83\xc2\x02\x45\x0f\xb7\x08\x41\x0f\xb7\xc1\x45\x85\xc9\x0f\x85\x6d\xff\xff\xff\x66\x41\x39\x30\x8b\xcf\x8b\xc7\x0f\x96\xc1\x66\x41\x39\x38\x0f\x95\xc0\x85\xc8\x74\x1a\x49\x83\xc0\x02\x8b\xd7\x8b\xcf"
+        b"\x66\x41\x39\x30\x0f\x96\xc2\x66\x41\x39\x38\x0f\x95\xc1\x85\xd1\x75\xe6\x48\x8b\x6c\x24\x10\x4c\x2b\xc3\x48\x8b\x5c\x24\x08\x48\x8b\x74\x24\x18\x49\xd1\xf8\x66\x41\x89\x3a\x41\x8b\xc0\x48\x8b\x7c\x24\x20\xc3\x48\x8b\xc4\x48\x89\x58\x08\x57\x48\x83\xec\x40"
+        b"\x48\x83\x60\xe8\x00\x48\x8b\xd9\x83\x21\x00\x45\x33\xc9\x83\x60\xe0\x00\x48\x8d\x0d\x87\x02\x00\x00\x45\x33\xc0\xc7\x40\xd8\x03\x00\x00\x00\xba\x00\x00\x00\xc0\xff\x15\x32\x02\x00\x00\x48\x8b\xf8\x48\x85\xc0\x7e\x30\x48\x8b\xc8\xff\x15\x29\x02\x00\x00\x83"
+        b"\xf8\x02\x75\x22\x48\x8b\xd3\x48\x8b\xcf\xff\x15\x38\x02\x00\x00\x85\xc0\x74\x12\x83\x0b\x04\x48\x8b\xcf\x8b\x13\xff\x15\x2e\x02\x00\x00\x48\x89\x7b\x08\x48\x8b\xc3\x48\x8b\x5c\x24\x50\x48\x83\xc4\x40\x5f\xc3\x48\x83\xec\x28\x48\x83\x79\x08\x00\x7e\x0c\x8b"
+        b"\x11\x48\x8b\x49\x08\xff\x15\x05\x02\x00\x00\x48\x83\xc4\x28\xc3\x48\x8b\xc4\x48\x89\x58\x08\x48\x89\x68\x10\x48\x89\x70\x18\x48\x89\x78\x20\x41\x55\x41\x56\x41\x57\x48\x83\xec\x30\x48\x8d\x0d\x2e\x05\x00\x00\xff\x15\xc6\x01\x00\x00\x45\x33\xff\x48\x85\xc0"
+        b"\x75\x09\x8d\x48\x3f\xe8\x7a\x01\x00\x00\xcc\x48\x8d\x15\xde\x01\x00\x00\x48\x8b\xc8\xff\x15\x9d\x01\x00\x00\x48\x8b\xe8\x48\x85\xc0\x75\x09\x8d\x48\x3e\xe8\x59\x01\x00\x00\xcc\x0f\xb7\x05\xed\x04\x00\x00\x4c\x8d\x2d\xe6\x04\x00\x00\x48\x8d\x0c\x45\xf8\xff"
+        b"\xff\xff\x49\x03\xcd\xeb\x04\x48\x83\xe9\x02\x66\x83\x79\xfe\x6e\x75\xf5\xc7\x01\x2e\x00\x65\x00\x48\x8d\x15\xc3\x04\x00\x00\xc7\x41\x04\x78\x00\x65\x00\x66\x44\x89\x79\x08\x48\x8d\x0d\x86\x01\x00\x00\xff\x15\x20\x01\x00\x00\x0f\xb7\x05\xa1\x04\x00\x00\x48"
+        b"\x8d\x1c\x45\x06\x00\x00\x00\x49\x03\xdd\x66\x44\x39\x7b\xfc\x75\x07\x66\x44\x39\x7b\xfe\x74\x06\x48\x83\xc3\x02\xeb\xec\xff\x15\xec\x00\x00\x00\x41\x8b\xd7\x41\x8b\xcf\x48\x8b\xf8\x66\x83\x38\x20\x0f\x96\xc2\x66\x44\x39\x38\x0f\x95\xc1\x85\xd1\x74\x1c\x48"
+        b"\x83\xc7\x02\x41\x8b\xcf\x41\x8b\xc7\x66\x83\x3f\x20\x0f\x96\xc1\x66\x44\x39\x3f\x0f\x95\xc0\x85\xc8\x75\xe4\x48\x8b\xcf\x48\x89\x3b\xe8\x32\xfd\xff\xff\x48\x63\xc8\xbe\x01\x00\x00\x00\x0f\xb7\x05\x2b\x04\x00\x00\x8b\xd6\x4c\x8d\x34\x4f\x49\x8d\x4d\x02\x41"
+        b"\xbd\x00\x01\x00\x00\x48\x8d\x0c\x41\x66\x44\x39\x39\x74\x26\x48\x89\x0c\xd3\x48\x83\xc1\x04\x66\x44\x39\x79\xfe\x74\x0d\x48\x8b\xc1\x48\x83\xc1\x02\x66\x44\x39\x38\x75\xf3\xff\xc6\x48\xff\xc2\x49\x3b\xd5\x7c\xd4\x48\x63\xfe\xeb\x1e\x66\x45\x39\x3e\x74\x1d"
+        b"\x49\x8b\xce\x4c\x89\x34\xfb\xe8\x94\xfb\xff\xff\x4c\x63\xc0\xff\xc6\x48\xff\xc7\x4f\x8d\x34\x46\x49\x3b\xfd\x7c\xdd\x48\x8d\x4c\x24\x20\xe8\xcd\xfd\xff\xff\x48\x8b\xd3\x8b\xce\xff\xd5\x48\x8d\x4c\x24\x20\x8b\xd8\xe8\x3a\xfe\xff\xff\x8b\xcb\xe8\x03\x00\x00"
+        b"\x00\xcc\xcc\xcc\x48\x83\xec\x28\xff\x15\x22\x00\x00\x00\xcc\x00\x98\x09\x00\x00\x00\x00\x00\x00\xaa\x09\x00\x00\x00\x00\x00\x00\xc4\x09\x00\x00\x00\x00\x00\x00\xd2\x09\x00\x00\x00\x00\x00\x00\xe0\x09\x00\x00\x00\x00\x00\x00\xee\x09\x00\x00\x00\x00\x00\x00"
+        b"\x00\x0a\x00\x00\x00\x00\x00\x00\x10\x0a\x00\x00\x00\x00\x00\x00\x22\x0a\x00\x00\x00\x00\x00\x00\x00\x00\x00\x00\x00\x00\x00\x00\x43\x00\x4f\x00\x4e\x00\x4f\x00\x55\x00\x54\x00\x24\x00\x00\x00\x50\x79\x5f\x4d\x61\x69\x6e\x00\x50\x00\x59\x00\x54\x00\x48\x00"
+        b"\x4f\x00\x4e\x00\x45\x00\x58\x00\x45\x00\x43\x00\x55\x00\x54\x00\x41\x00\x42\x00\x4c\x00\x45\x00\x00\x00\x00\x00\x00\x00\x00\x00\xe3\xc2\x11\x64\x00\x00\x00\x00\x0d\x00\x00\x00\xec\x00\x00\x00\xd8\x07\x00\x00\xd8\x07\x00\x00\x18\x00\x00\x00\x02\x80\x02\x80"
+        b"\x00\x00\x00\x00\x00\x00\x00\x00\xd0\x07\x00\x00\x08\x00\x00\x00\x60\x02\x00\x00\xaf\x04\x00\x00\x00\x00\x00\x00\x60\x02\x00\x00\xaf\x04\x00\x00\x2e\x74\x65\x78\x74\x24\x6d\x6e\x00\x00\x00\x00\x10\x07\x00\x00\x50\x00\x00\x00\x2e\x69\x64\x61\x74\x61\x24\x35"
+        b"\x00\x00\x00\x00\x60\x07\x00\x00\x58\x00\x00\x00\x2e\x72\x64\x61\x74\x61\x00\x00\xb8\x07\x00\x00\x20\x00\x00\x00\x2e\x72\x64\x61\x74\x61\x24\x76\x6f\x6c\x74\x6d\x64\x00\x00\x00\xd8\x07\x00\x00\xec\x00\x00\x00\x2e\x72\x64\x61\x74\x61\x24\x7a\x7a\x7a\x64\x62"
+        b"\x67\x00\x00\x00\xc4\x08\x00\x00\x5c\x00\x00\x00\x2e\x78\x64\x61\x74\x61\x00\x00\x20\x09\x00\x00\x14\x00\x00\x00\x2e\x69\x64\x61\x74\x61\x24\x32\x00\x00\x00\x00\x34\x09\x00\x00\x14\x00\x00\x00\x2e\x69\x64\x61\x74\x61\x24\x33\x00\x00\x00\x00\x48\x09\x00\x00"
+        b"\x50\x00\x00\x00\x2e\x69\x64\x61\x74\x61\x24\x34\x00\x00\x00\x00\x98\x09\x00\x00\xaa\x00\x00\x00\x2e\x69\x64\x61\x74\x61\x24\x36\x00\x00\x00\x00\x50\x0a\x00\x00\x48\x00\x00\x00\x2e\x70\x64\x61\x74\x61\x00\x00\xa0\x0a\x00\x00\x04\x00\x00\x00\x2e\x73\x68\x69"
+        b"\x6d\x00\x00\x00\x01\x04\x01\x00\x04\x42\x00\x00\x01\x0c\x04\x00\x0c\x34\x0a\x00\x0c\x72\x08\x70\x01\x1d\x0c\x00\x1d\x74\x0d\x00\x1d\x64\x0c\x00\x1d\x54\x0b\x00\x1d\x34\x0a\x00\x1d\x52\x19\xf0\x17\xe0\x15\xd0\x01\x13\x08\x00\x13\x74\x04\x00\x0f\x64\x03\x00"
+        b"\x0b\x54\x02\x00\x07\x34\x01\x00\x01\x15\x09\x00\x15\x74\x05\x00\x15\x64\x04\x00\x15\x54\x03\x00\x15\x34\x02\x00\x15\xe0\x00\x00\x48\x09\x00\x00\x00\x00\x00\x00\x00\x00\x00\x00\x34\x0a\x00\x00\x10\x07\x00\x00\x00\x00\x00\x00\x00\x00\x00\x00\x00\x00\x00\x00"
+        b"\x00\x00\x00\x00\x00\x00\x00\x00\x98\x09\x00\x00\x00\x00\x00\x00\xaa\x09\x00\x00\x00\x00\x00\x00\xc4\x09\x00\x00\x00\x00\x00\x00\xd2\x09\x00\x00\x00\x00\x00\x00\xe0\x09\x00\x00\x00\x00\x00\x00\xee\x09\x00\x00\x00\x00\x00\x00\x00\x0a\x00\x00\x00\x00\x00\x00"
+        b"\x10\x0a\x00\x00\x00\x00\x00\x00\x22\x0a\x00\x00\x00\x00\x00\x00\x00\x00\x00\x00\x00\x00\x00\x00\xf1\x01\x47\x65\x74\x43\x6f\x6d\x6d\x61\x6e\x64\x4c\x69\x6e\x65\x57\x00\x46\x05\x53\x65\x74\x45\x6e\x76\x69\x72\x6f\x6e\x6d\x65\x6e\x74\x56\x61\x72\x69\x61\x62"
+        b"\x6c\x65\x57\x00\xda\x00\x43\x72\x65\x61\x74\x65\x46\x69\x6c\x65\x57\x00\x6a\x02\x47\x65\x74\x46\x69\x6c\x65\x54\x79\x70\x65\x00\x78\x01\x45\x78\x69\x74\x50\x72\x6f\x63\x65\x73\x73\x00\xcd\x02\x47\x65\x74\x50\x72\x6f\x63\x41\x64\x64\x72\x65\x73\x73\x00\x00"
+        b"\xe7\x03\x4c\x6f\x61\x64\x4c\x69\x62\x72\x61\x72\x79\x57\x00\x00\x16\x02\x47\x65\x74\x43\x6f\x6e\x73\x6f\x6c\x65\x4d\x6f\x64\x65\x00\x00\x2b\x05\x53\x65\x74\x43\x6f\x6e\x73\x6f\x6c\x65\x4d\x6f\x64\x65\x00\x00\x4b\x45\x52\x4e\x45\x4c\x33\x32\x2e\x64\x6c\x6c"
+        b"\x00\x00\x00\x00\x00\x00\x00\x00\x00\x00\x00\x00\x00\x00\x00\x00\x60\x02\x00\x00\x98\x03\x00\x00\x08\x09\x00\x00\x98\x03\x00\x00\xb4\x04\x00\x00\xf4\x08\x00\x00\xb4\x04\x00\x00\x34\x05\x00\x00\xcc\x08\x00\x00\x34\x05\x00\x00\x50\x05\x00\x00\xc4\x08\x00\x00"
+        b"\x50\x05\x00\x00\x02\x07\x00\x00\xd8\x08\x00\x00\x04\x07\x00\x00\x0f\x07\x00\x00\xc4\x08\x00\x00\x00\x00\x00\x00\x00\x00\x00\x00"
+    )
+    shim_section = 568
+    shim_subsystem = 276
+    shim_alignment = 16
+    shim_data = 2256
+    shim_payload = (
+        b"\x4d\x5a\x90\x00\x03\x00\x00\x00\x04\x00\x00\x00\xff\xff\x00\x00\xb8\x00\x00\x00\x00\x00\x00\x00\x40\x00\x00\x00\x00\x00\x00\x00\x00\x00\x00\x00\x00\x00\x00\x00\x00\x00\x00\x00\x00\x00\x00\x00\x00\x00\x00\x00\x00\x00\x00\x00\x00\x00\x00\x00\xb8\x00\x00\x00"
+        b"\x0e\x1f\xba\x0e\x00\xb4\x09\xcd\x21\xb8\x01\x4c\xcd\x21\x54\x68\x69\x73\x20\x70\x72\x6f\x67\x72\x61\x6d\x20\x63\x61\x6e\x6e\x6f\x74\x20\x62\x65\x20\x72\x75\x6e\x20\x69\x6e\x20\x44\x4f\x53\x20\x6d\x6f\x64\x65\x2e\x0d\x0d\x0a\x24\x00\x00\x00\x00\x00\x00\x00"
+        b"\xf1\x02\x04\xcf\xb5\x63\x6a\x9c\xb5\x63\x6a\x9c\xb5\x63\x6a\x9c\xa1\x08\x6b\x9d\xb6\x63\x6a\x9c\xb5\x63\x6b\x9c\xbe\x63\x6a\x9c\x08\x16\x6f\x9d\xb4\x63\x6a\x9c\x08\x16\x68\x9d\xb4\x63\x6a\x9c\x52\x69\x63\x68\xb5\x63\x6a\x9c\x50\x45\x00\x00\x64\x86\x04\x00"
+        b"\xba\x68\x12\x63\x00\x00\x00\x00\x00\x00\x00\x00\xf0\x00\x22\x00\x0b\x02\x0e\x1d\x00\x03\x00\x00\x80\x03\x00\x00\x00\x00\x00\x00\x1c\x03\x00\x00\x60\x02\x00\x00\x00\x00\x00\x40\x01\x00\x00\x00\x10\x00\x00\x00\x10\x00\x00\x00\x06\x00\x00\x00\x00\x00\x00\x00"
+        b"\x06\x00\x00\x00\x00\x00\x00\x00\x00\x40\x00\x00\x60\x02\x00\x00\x00\x00\x00\x00\x03\x00\x60\x81\x00\x00\x10\x00\x00\x00\x00\x00\x00\x10\x00\x00\x00\x00\x00\x00\x00\x00\x10\x00\x00\x00\x00\x00\x00\x10\x00\x00\x00\x00\x00\x00\x00\x00\x00\x00\x10\x00\x00\x00"
+        b"\x00\x00\x00\x00\x00\x00\x00\x00\x2c\x07\x00\x00\x28\x00\x00\x00\x00\x00\x00\x00\x00\x00\x00\x00\xa0\x08\x00\x00\x30\x00\x00\x00\x00\x00\x00\x00\x00\x00\x00\x00\x00\x00\x00\x00\x00\x00\x00\x00\xd4\x05\x00\x00\x1c\x00\x00\x00\x00\x00\x00\x00\x00\x00\x00\x00"
+        b"\x00\x00\x00\x00\x00\x00\x00\x00\x00\x00\x00\x00\x00\x00\x00\x00\x00\x00\x00\x00\x00\x00\x00\x00\x00\x00\x00\x00\x00\x00\x00\x00\x60\x05\x00\x00\x60\x00\x00\x00\x00\x00\x00\x00\x00\x00\x00\x00\x00\x00\x00\x00\x00\x00\x00\x00\x00\x00\x00\x00\x00\x00\x00\x00"
+        b"\x2e\x74\x65\x78\x74\x00\x00\x00\xf4\x02\x00\x00\x60\x02\x00\x00\x00\x03\x00\x00\x60\x02\x00\x00\x00\x00\x00\x00\x00\x00\x00\x00\x00\x00\x00\x00\x20\x00\x00\x60\x2e\x72\x64\x61\x74\x61\x00\x00\x32\x03\x00\x00\x60\x05\x00\x00\x40\x03\x00\x00\x60\x05\x00\x00"
+        b"\x00\x00\x00\x00\x00\x00\x00\x00\x00\x00\x00\x00\x40\x00\x00\x40\x2e\x70\x64\x61\x74\x61\x00\x00\x30\x00\x00\x00\xa0\x08\x00\x00\x30\x00\x00\x00\xa0\x08\x00\x00\x00\x00\x00\x00\x00\x00\x00\x00\x00\x00\x00\x00\x40\x00\x00\x40\x2e\x73\x68\x69\x6d\x00\x00\x00"
+        b"\x06\x00\x00\x00\xd0\x08\x00\x00\x10\x00\x00\x00\xd0\x08\x00\x00\x00\x00\x00\x00\x00\x00\x00\x00\x00\x00\x00\x00\x40\x00\x00\xc0\x33\xc0\x83\xf9\x02\x0f\x92\xc0\xc3\xcc\xcc\xcc\x48\x8b\xc4\x48\x89\x58\x08\x57\x48\x83\xec\x40\x48\x83\x60\xe8\x00\x48\x8b\xd9"
+        b"\x83\x21\x00\x45\x33\xc9\x83\x60\xe0\x00\x48\x8d\x0d\x2f\x03\x00\x00\x45\x33\xc0\xc7\x40\xd8\x03\x00\x00\x00\xba\x00\x00\x00\xc0\xff\x15\xc2\x02\x00\x00\x48\x8b\xf8\x48\x85\xc0\x7e\x30\x48\x8b\xc8\xff\x15\xb9\x02\x00\x00\x83\xf8\x02\x75\x22\x48\x8b\xd3\x48"
+        b"\x8b\xcf\xff\x15\xd8\x02\x00\x00\x85\xc0\x74\x12\x83\x0b\x04\x48\x8b\xcf\x8b\x13\xff\x15\xce\x02\x00\x00\x48\x89\x7b\x08\x48\x8b\xc3\x48\x8b\x5c\x24\x50\x48\x83\xc4\x40\x5f\xc3\x48\x83\xec\x28\x48\x83\x79\x08\x00\x7e\x0c\x8b\x11\x48\x8b\x49\x08\xff\x15\xa5"
+        b"\x02\x00\x00\x48\x83\xc4\x28\xc3\x48\x8d\x05\x51\xff\xff\xff\xc3\x48\x83\xec\x28\xff\x15\x66\x02\x00\x00\xcc\xcc\x4c\x8b\xdc\x49\x89\x5b\x08\x49\x89\x6b\x18\x56\x57\x41\x54\x41\x56\x41\x57\x48\x81\xec\xf0\x00\x00\x00\x48\x8d\x0d\x93\x02\x00\x00\x48\x81\xe1"
+        b"\x00\x00\xff\xff\x48\x63\x41\x3c\x0f\xb7\x6c\x08\x5c\x66\x83\xfd\x03\x75\x17\x49\x8d\x4b\x10\xe8\xac\xff\xff\xff\x48\x8b\xc8\xba\x01\x00\x00\x00\xff\x15\x46\x02\x00\x00\x33\xdb\x48\x8d\x35\x61\x05\x00\x00\x44\x8a\xc3\x48\x8b\xc6\x44\x8d\x63\x02\x0f\xb7\x08"
+        b"\x66\x85\xc9\x74\x11\x66\x83\xf9\x2f\x41\x0f\x94\xc0\x49\x03\xc4\x66\x83\xf9\x2f\x75\xe7\x66\x83\x3d\x38\x05\x00\x00\x3a\x75\x07\xba\x01\x00\x00\x00\xeb\x09\x45\x84\xc0\x8b\xd3\x41\x0f\x45\xd4\x0f\xb7\x0d\x19\x05\x00\x00\x4c\x8d\x05\x12\x05\x00\x00\x0f\xb7"
+        b"\x3d\x0d\x05\x00\x00\x4c\x8b\xfe\x49\x03\xfc\x48\x03\xf9\x4c\x8d\x71\x02\x4f\x8d\x34\x70\x49\x8d\x3c\x78\x41\x3b\xd4\x75\x36\x41\xb8\x00\x40\x00\x00\x48\x8b\xd7\x33\xc9\x4c\x8b\xff\xff\x15\xa5\x01\x00\x00\x8b\xc0\x48\x8d\x3c\x47\x49\x2b\xfc\x66\x83\x7f\xfe"
+        b"\x5c\x75\xf6\x48\x2b\xf7\x0f\xb7\x04\x3e\x66\x89\x07\x49\x03\xfc\x66\x85\xc0\x75\xf1\x48\x8b\xf7\xff\x15\x42\x01\x00\x00\x41\xb9\x22\x00\x00\x00\x48\x8b\xd0\x66\x44\x39\x08\x75\x30\x45\x8b\xc4\x41\x0f\xb7\xc1\x49\x03\xd4\x66\x89\x06\x49\x03\xf4\x8b\xcb\x66"
+        b"\x44\x39\x4a\xfe\x0f\x94\xc1\x44\x2b\xc1\x74\x33\x0f\xb7\x02\x66\x85\xc0\x75\xe0\x66\x44\x89\x0e\x49\x03\xf4\xeb\x22\x48\x8b\xcf\x48\x2b\xc8\x0f\xb7\x02\x48\x8d\x34\x11\x66\x89\x06\x4c\x8b\xc2\x0f\xb7\x02\x49\x03\xd4\x66\x83\xf8\x20\x77\xe7\x49\x8b\xd0\x49"
+        b"\x2b\xf6\x41\x0f\xb7\x06\x4a\x8d\x0c\x36\x66\x89\x01\x41\x0f\xb7\x06\x4d\x03\xf4\x66\x85\xc0\x75\xe9\x0f\xb7\x02\x49\x03\xd4\x66\x89\x01\x49\x03\xcc\x0f\xb7\x02\x66\x85\xc0\x75\xef\x48\x8d\x4c\x24\x50\xe8\xb5\xfd\xff\xff\x48\x8d\x44\x24\x60\xc7\x84\x24\x80"
+        b"\x00\x00\x00\x68\x00\x00\x00\x48\x89\x44\x24\x48\x45\x33\xc9\x48\x8d\x84\x24\x80\x00\x00\x00\x45\x33\xc0\x48\x89\x44\x24\x40\x48\x8b\xd7\x48\x89\x5c\x24\x38\x49\x8b\xcf\x48\x89\x5c\x24\x30\xc7\x44\x24\x28\x00\x00\x01\x00\x89\x5c\x24\x20\xff\x15\x8f\x00\x00"
+        b"\x00\x83\xf8\x01\x75\x37\x66\x83\xfd\x03\x75\x36\x48\x8b\x4c\x24\x60\x83\xca\xff\xff\x15\x5e\x00\x00\x00\x48\x8b\x4c\x24\x60\x48\x8d\x94\x24\x28\x01\x00\x00\x89\x9c\x24\x28\x01\x00\x00\xff\x15\x54\x00\x00\x00\x8b\x9c\x24\x28\x01\x00\x00\xeb\x05\xbb\x31\x23"
+        b"\x00\x00\x48\x8d\x4c\x24\x50\xe8\xa0\xfd\xff\xff\x8b\xcb\xe8\xbd\xfd\xff\xff\xcc\x00\x00\x00\x00\x00\x00\x00\x00\x00\x00\x00\x00\xb8\x07\x00\x00\x00\x00\x00\x00\xca\x07\x00\x00\x00\x00\x00\x00\xd8\x07\x00\x00\x00\x00\x00\x00\xe6\x07\x00\x00\x00\x00\x00\x00"
+        b"\xfc\x07\x00\x00\x00\x00\x00\x00\x0a\x08\x00\x00\x00\x00\x00\x00\x20\x08\x00\x00\x00\x00\x00\x00\x32\x08\x00\x00\x00\x00\x00\x00\x48\x08\x00\x00\x00\x00\x00\x00\x5a\x08\x00\x00\x00\x00\x00\x00\x6c\x08\x00\x00\x00\x00\x00\x00\x00\x00\x00\x00\x00\x00\x00\x00"
+        b"\x43\x00\x4f\x00\x4e\x00\x4f\x00\x55\x00\x54\x00\x24\x00\x00\x00\x00\x00\x00\x00\x00\x00\x00\x00\xba\x68\x12\x63\x00\x00\x00\x00\x0d\x00\x00\x00\xec\x00\x00\x00\x10\x06\x00\x00\x10\x06\x00\x00\x18\x00\x00\x00\x02\x80\x02\x80\x00\x00\x00\x00\x00\x00\x00\x00"
+        b"\x08\x06\x00\x00\x08\x00\x00\x00\x60\x02\x00\x00\xf4\x02\x00\x00\x00\x00\x00\x00\x60\x02\x00\x00\xf4\x02\x00\x00\x2e\x74\x65\x78\x74\x24\x6d\x6e\x00\x00\x00\x00\x60\x05\x00\x00\x60\x00\x00\x00\x2e\x69\x64\x61\x74\x61\x24\x35\x00\x00\x00\x00\xc0\x05\x00\x00"
+        b"\x30\x00\x00\x00\x2e\x72\x64\x61\x74\x61\x00\x00\xf0\x05\x00\x00\x20\x00\x00\x00\x2e\x72\x64\x61\x74\x61\x24\x76\x6f\x6c\x74\x6d\x64\x00\x00\x00\x10\x06\x00\x00\xec\x00\x00\x00\x2e\x72\x64\x61\x74\x61\x24\x7a\x7a\x7a\x64\x62\x67\x00\x00\x00\xfc\x06\x00\x00"
+        b"\x30\x00\x00\x00\x2e\x78\x64\x61\x74\x61\x00\x00\x2c\x07\x00\x00\x14\x00\x00\x00\x2e\x69\x64\x61\x74\x61\x24\x32\x00\x00\x00\x00\x40\x07\x00\x00\x18\x00\x00\x00\x2e\x69\x64\x61\x74\x61\x24\x33\x00\x00\x00\x00\x58\x07\x00\x00\x60\x00\x00\x00\x2e\x69\x64\x61"
+        b"\x74\x61\x24\x34\x00\x00\x00\x00\xb8\x07\x00\x00\xda\x00\x00\x00\x2e\x69\x64\x61\x74\x61\x24\x36\x00\x00\x00\x00\xa0\x08\x00\x00\x30\x00\x00\x00\x2e\x70\x64\x61\x74\x61\x00\x00\xd0\x08\x00\x00\x06\x00\x00\x00\x2e\x73\x68\x69\x6d\x00\x00\x00\x01\x04\x01\x00"
+        b"\x04\x42\x00\x00\x01\x0c\x04\x00\x0c\x34\x0a\x00\x0c\x72\x08\x70\x01\x1a\x0b\x00\x1a\x54\x26\x00\x1a\x34\x24\x00\x1a\x01\x1e\x00\x13\xf0\x11\xe0\x0f\xc0\x0d\x70\x0c\x60\x00\x00\x58\x07\x00\x00\x00\x00\x00\x00\x00\x00\x00\x00\x84\x08\x00\x00\x60\x05\x00\x00"
+        b"\x00\x00\x00\x00\x00\x00\x00\x00\x00\x00\x00\x00\x00\x00\x00\x00\x00\x00\x00\x00\x00\x00\x00\x00\xb8\x07\x00\x00\x00\x00\x00\x00\xca\x07\x00\x00\x00\x00\x00\x00\xd8\x07\x00\x00\x00\x00\x00\x00\xe6\x07\x00\x00\x00\x00\x00\x00\xfc\x07\x00\x00\x00\x00\x00\x00"
+        b"\x0a\x08\x00\x00\x00\x00\x00\x00\x20\x08\x00\x00\x00\x00\x00\x00\x32\x08\x00\x00\x00\x00\x00\x00\x48\x08\x00\x00\x00\x00\x00\x00\x5a\x08\x00\x00\x00\x00\x00\x00\x6c\x08\x00\x00\x00\x00\x00\x00\x00\x00\x00\x00\x00\x00\x00\x00\xe0\x01\x47\x65\x74\x43\x6f\x6d"
+        b"\x6d\x61\x6e\x64\x4c\x69\x6e\x65\x57\x00\xce\x00\x43\x72\x65\x61\x74\x65\x46\x69\x6c\x65\x57\x00\x58\x02\x47\x65\x74\x46\x69\x6c\x65\x54\x79\x70\x65\x00\xea\x05\x57\x61\x69\x74\x46\x6f\x72\x53\x69\x6e\x67\x6c\x65\x4f\x62\x6a\x65\x63\x74\x00\x67\x01\x45\x78"
+        b"\x69\x74\x50\x72\x6f\x63\x65\x73\x73\x00\x46\x02\x47\x65\x74\x45\x78\x69\x74\x43\x6f\x64\x65\x50\x72\x6f\x63\x65\x73\x73\x00\x00\xe8\x00\x43\x72\x65\x61\x74\x65\x50\x72\x6f\x63\x65\x73\x73\x57\x00\x00\x7d\x02\x47\x65\x74\x4d\x6f\x64\x75\x6c\x65\x46\x69\x6c"
+        b"\x65\x4e\x61\x6d\x65\x57\x00\x00\x05\x02\x47\x65\x74\x43\x6f\x6e\x73\x6f\x6c\x65\x4d\x6f\x64\x65\x00\x00\x09\x05\x53\x65\x74\x43\x6f\x6e\x73\x6f\x6c\x65\x4d\x6f\x64\x65\x00\x00\xf9\x04\x53\x65\x74\x43\x6f\x6e\x73\x6f\x6c\x65\x43\x74\x72\x6c\x48\x61\x6e\x64"
+        b"\x6c\x65\x72\x00\x4b\x45\x52\x4e\x45\x4c\x33\x32\x2e\x64\x6c\x6c\x00\x00\x00\x00\x00\x00\x00\x00\x00\x00\x00\x00\x00\x00\x00\x00\x6c\x02\x00\x00\xec\x02\x00\x00\x04\x07\x00\x00\xec\x02\x00\x00\x08\x03\x00\x00\xfc\x06\x00\x00\x10\x03\x00\x00\x1b\x03\x00\x00"
+        b"\xfc\x06\x00\x00\x1c\x03\x00\x00\x54\x05\x00\x00\x10\x07\x00\x00"
+    )
+
+#-------------------------------------------------------------------------------
+def create_py(shim_path, *args):
+    dll_path = os.path.dirname(sys.executable)
+    dll_path += rf"\python{sys.version_info.major}{sys.version_info.minor}.dll"
+
+    out = open(shim_path, "wb")
+    out.write(_K.py_payload)
+
+    out.seek(_K.py_data)
+
+    out.write(struct.pack("H", len(dll_path) + 1)) # +1 for null terminator
+    out.write(dll_path.encode("utf16")[2:])        # "[2:]" strips the BOM
+    out.write(struct.pack("H", 0))
+
+    for arg in args:
+        out.write(arg.encode("utf16")[2:])
+        out.write(struct.pack("H", 0))
+    out.write(struct.pack("H", 0))
+
+    align_mask = _K.py_alignment - 1
+    padding = out.tell() - _K.py_data
+    padding = ((padding + align_mask) & ~align_mask) - padding
+    for i in range(padding):
+        out.write(b"\0")
+
+    section_size = out.tell() - _K.py_data
+    out.seek(_K.py_section + 16)
+    out.write(struct.pack("I", section_size))
+
+    out.close()
+
+#-------------------------------------------------------------------------------
+def create(shim_path, target_path, args=""):
+    # Condition input
+    target_path = target_path.replace("\\", "/")
+    if args:
+        args = " " + args
+
+    out = open(shim_path, "wb")
+
+    # Get the target's subsystem.
+    full_target_path = os.path.join(os.path.dirname(out.name), target_path)
+    with open(full_target_path, "rb") as exe_file:
+        # DOS header
+        dos_header = exe_file.read(64)
+        dos_header = struct.unpack("H" + ("H" * 29) + "I", dos_header)
+        pe_offset = dos_header[-1]
+        assert dos_header[0] == 0x5a4d
+
+        # NT header
+        exe_file.seek(pe_offset)
+        pe_magic, = struct.unpack("I", exe_file.read(4))
+        assert pe_magic == 0x4550 # 'PE'
+
+        # File header
+        file_header = struct.unpack("HHIIIHH", exe_file.read(20))
+        optional_header_size = file_header[-2] - (8 * 16)
+
+        # Optional header
+        optional_header = exe_file.read(optional_header_size)
+        if file_header[0] == 0x8664: opt_unpack = "HBBIIIIIQIIHHHHHHIIIIHHQQQQII"
+        else: opt_unpack = "HBBIIIIIIIIIHHHHHHIIIIHHIIIIII"
+        optional_header = struct.unpack(opt_unpack, optional_header)
+        subsystem = optional_header[-8]
+        assert subsystem == 2 or subsystem == 3
+
+    out.write(_K.shim_payload)
+
+    out.seek(_K.shim_subsystem)
+    out.write(struct.pack("H", subsystem))
+
+    out.seek(_K.shim_data)
+
+    out.write(struct.pack("H", len(target_path) + 1)) # +1 for null terminator
+    out.write(struct.pack("H", len(args) + 1))        # +1 for null terminator
+
+    out.write(target_path.encode("utf16")[2:])        # "[2:]" strips the BOM
+    out.write(struct.pack("H", 0))
+
+    out.write(args.encode("utf16")[2:])               # "[2:]" strips the BOM
+    out.write(struct.pack("H", 0))
+
+    align_mask = _K.shim_alignment - 1
+    padding = out.tell() - _K.shim_data
+    padding = ((padding + align_mask) & ~align_mask) - padding
+    for i in range(padding):
+        out.write(b"\0")
+
+    section_size = out.tell() - _K.shim_data
+    out.seek(_K.shim_section + 16)
+    out.write(struct.pack("I", section_size))
+
+    out.close()
diff --git a/Engine/Extras/ushell/channels/flow/core/system/lib/exeshim/build.bat b/Engine/Extras/ushell/channels/flow/core/system/lib/exeshim/build.bat
new file mode 100755
--- /dev/null
+++ b/Engine/Extras/ushell/channels/flow/core/system/lib/exeshim/build.bat
@@ -0,0 +1,38 @@
+:: Copyright Epic Games, Inc. All Rights Reserved.
+
+@echo off
+setlocal
+
+cl.exe 1>nul 2>nul
+if errorlevel 9009 (
+    for /f "delims=" %%d in ('vswhere -latest -find VC/Auxiliary/Build/vcvars64.bat') do (
+        call "%%d"
+    )
+)
+
+pushd %~dp0
+1>nul 2>nul mkdir _build
+cd _build
+
+if "%1"=="debug" (
+    set _cl=/Od /DDEBUG
+    set _link=/debug
+)
+call:build 1 py
+call:build 2 shim
+call:build 3 dump
+
+dump.exe
+if errorlevel 1 (
+    echo !! dump failed!
+)
+
+popd
+goto:eof
+
+:build
+echo.
+echo shim.cpp / %2 / %1
+cl ../shim.cpp /Fo%2.obj /DSHIM_BUILD=%1 /nologo /FC /Zi /Zl /EHs-c- /O1 /c %_cl%
+link %_link% /nologo /nodefaultlib /align:16 /entry:%2 /out:%2.exe /subsystem:console %2.obj kernel32.lib user32.lib shell32.lib
+goto:eof
diff --git a/Engine/Extras/ushell/channels/flow/core/system/lib/exeshim/shim.cpp b/Engine/Extras/ushell/channels/flow/core/system/lib/exeshim/shim.cpp
new file mode 100644
--- /dev/null
+++ b/Engine/Extras/ushell/channels/flow/core/system/lib/exeshim/shim.cpp
@@ -0,0 +1,474 @@
+// Copyright Epic Games, Inc. All Rights Reserved.
+
+#include <cinttypes>
+#include <Windows.h>
+#include <winternl.h>
+
+// {{{1 misc -------------------------------------------------------------------
+static void close_handle(void*)
+{
+    /* do nothing - let the kernel clean up after us... */
+}
+
+[[noreturn]] static void quit(int exit_code)
+{
+    ExitProcess(exit_code);
+}
+
+static void spam(const char* format, ...)
+{
+    HANDLE conout = CreateFileW(L"CONOUT$", GENERIC_WRITE|GENERIC_READ, 0, nullptr, OPEN_EXISTING, 0, nullptr);
+
+    static char buffer[128];
+
+    va_list args;
+    va_start(args, format);
+    int count = wvsprintfA(buffer, format, args);
+    va_end(args);
+
+    DWORD written;
+    WriteFile(conout, buffer, count, &written, nullptr);
+
+    CloseHandle(conout);
+};
+
+#if 0
+#   define SPAM(format, ...) spam(format "\n", ##__VA_ARGS__)
+#else
+#   define SPAM(format, ...)
+#endif
+
+// {{{1 constants --------------------------------------------------------------
+static const uint16_t LENGTH_ID = 0x493;
+static const uint16_t DATA_ID   = 493;
+
+// {{{1 vt100 ------------------------------------------------------------------
+class Vt100Scope
+{
+public:
+            Vt100Scope();
+            ~Vt100Scope();
+
+private:
+    DWORD   _con_mode;
+    HANDLE  _stdout_handle;
+};
+
+//------------------------------------------------------------------------------
+Vt100Scope::Vt100Scope()
+{
+    // Enable Windows 10's VT100/ECMA-48 support
+    _con_mode = 0;
+    HANDLE handle = CreateFileW(L"CONOUT$", GENERIC_WRITE|GENERIC_READ, 0, nullptr, OPEN_EXISTING, 0, nullptr);
+    if (intptr_t(handle) <= 0)
+        return;
+
+    if (GetFileType(handle) != FILE_TYPE_CHAR)
+        return;
+
+    if (!GetConsoleMode(handle, &_con_mode))
+        return;
+
+    _con_mode |= ENABLE_VIRTUAL_TERMINAL_PROCESSING;
+    SetConsoleMode(handle, _con_mode);
+
+    _stdout_handle = handle;
+}
+
+//------------------------------------------------------------------------------
+Vt100Scope::~Vt100Scope()
+{
+    if (intptr_t(_stdout_handle) <= 0)
+        return;
+
+    // Restore modifications to the terminal's mode
+    SetConsoleMode(_stdout_handle, _con_mode);
+}
+
+// {{{1 py ---------------------------------------------------------------------
+#if SHIM_BUILD == 1
+
+#pragma section(".shim", read, write)
+
+struct PyData
+{
+    uint16_t    dll_length;
+    wchar_t     dll_path[1];
+};
+
+#if defined(DEBUG)
+#   define PY_DLL L"[REDACTED_USER_HOME]/AppData/Local/ushell/.working/python/3.11.0/python311.dll"
+    static wchar_t _data[2048] = {
+        (sizeof(PY_DLL) / sizeof(PY_DLL[0])),
+        PY_DLL L"\0" L"-u\0" L"-E\0",
+    };
+#else
+    static __declspec(allocate(".shim")) PyData _data = {
+        LENGTH_ID,
+        { DATA_ID },
+    };
+#endif
+
+//------------------------------------------------------------------------------
+template <int EnableSingleQuote>
+static int read_arg_word(wchar_t* start)
+{
+    // note to self; the escaping of quotes doesn't quite match what MS does.
+
+    auto is_quote = [] (int c) {
+        return (c == '\"') | (EnableSingleQuote & (c == '\''));
+    };
+
+    wchar_t* write = start;
+    const wchar_t* read = start;
+    int c;
+    int quote = 0;
+    for (; c = *read; ++read, ++write)
+    {
+        *write = *read;
+
+        if (quote)
+        {
+            if (c == '\\' & read[1] == quote)
+            {
+                ++read;
+                *write = wchar_t(quote);
+                continue;
+            }
+
+            quote = (c == quote) ? 0 : quote;
+            write -= !quote;
+            continue;
+        }
+
+        if (c == '\\' & is_quote(read[1]))
+        {
+            ++read;
+            *write = '\"';
+            continue;
+        }
+        if (quote = (is_quote(c) ? c : 0))
+        {
+            --write;
+            continue;
+        }
+
+        if (c <= 0x20)
+            break;
+    }
+
+    for (; c = *read, (!!c & unsigned(c) <= 0x20); ++read);
+
+    *write = '\0';
+
+    return int(ptrdiff_t(read - start));
+}
+
+//------------------------------------------------------------------------------
+void __cdecl py(void*)
+{
+    auto* data = (PyData*)(&_data);
+
+    // Load Python DLL and find Py_Main entry point
+    SPAM("loading dll; %S", data->dll_path);
+    HMODULE py_dll = LoadLibraryW(data->dll_path);
+    if (py_dll == nullptr)
+        quit(63);
+    SPAM("loaded at %p", py_dll);
+
+    FARPROC py_main_addr = GetProcAddress(py_dll, "Py_Main");
+    if (py_main_addr == nullptr)
+        quit(62);
+    SPAM("Py_Main address; %p", py_main_addr);
+
+    // Convert DLL path into a .exe path
+    wchar_t* cursor = data->dll_path + data->dll_length - 5; // 5 skips ".dll\0"
+    for (; cursor[-1] != 'n'; --cursor);
+    cursor[0] = '.';
+    cursor[1] = 'e';
+    cursor[2] = 'x';
+    cursor[3] = 'e';
+    cursor[4] = '\0';
+
+    // Setting this makes Python behave as if the shim is python.exe
+    SetEnvironmentVariableW(L"PYTHONEXECUTABLE", data->dll_path);
+
+    // Build argv arguments
+    // We'll use memory after arguments for argv[]
+    cursor = data->dll_path + data->dll_length + 2;
+    for (; cursor[-2] || cursor[-1]; ++cursor);
+    wchar_t** argv = (wchar_t**)cursor;
+
+    wchar_t* cmd_line = GetCommandLineW();
+    for (int c; c = *cmd_line, (!!c & unsigned(c) <= 0x20); ++cmd_line);
+    argv[0] = cmd_line;
+    cmd_line += read_arg_word<0>(cmd_line);
+
+    int argv_i = 1;
+    const int argv_n = 256;
+
+    cursor = data->dll_path + data->dll_length;
+    for (; argv_i < argv_n && *cursor; ++argv_i)
+    {
+        argv[argv_i] = cursor;
+        for (cursor += 2; cursor[-1] != '\0'; ++cursor);
+    }
+
+    for (int i = 1; argv_i < argv_n && cmd_line[0] != '\0'; ++argv_i, ++i)
+    {
+        argv[argv_i] = cmd_line;
+        cmd_line += read_arg_word<1>(cmd_line);
+    }
+
+    // Launch Python
+    int ret = 0;
+    {
+        Vt100Scope _;
+
+        for (int i = 0; i < argv_i; ++i)
+            SPAM("arg %d %S", i, argv[i]);
+
+        auto* py_main = (int (*)(int, wchar_t**))py_main_addr;
+        ret = py_main(argv_i, argv);
+    }
+
+    quit(ret);
+}
+
+#endif // SHIM_BUILD
+
+// {{{1 shim -------------------------------------------------------------------
+#if SHIM_BUILD == 2
+
+#pragma section(".shim", read, write)
+
+struct ShimData
+{
+    uint16_t    path_length;
+    uint16_t    args_length;
+    wchar_t     path[1];
+};
+
+#if defined(DEBUG)
+#   define PATH L"c:/windows/system32/cmd.exe"
+#   define ARGS L"/c dir c:\\"
+    static wchar_t _data[2048] = {
+        (sizeof(PATH) / sizeof(PATH[0])),
+        (sizeof(ARGS) / sizeof(ARGS[0])),
+        PATH L"\0" ARGS,
+    };
+#else
+    static __declspec(allocate(".shim")) ShimData _data = {
+        LENGTH_ID,
+        LENGTH_ID,
+        { DATA_ID },
+    };
+#endif
+
+//------------------------------------------------------------------------------
+void __cdecl shim(void*)
+{
+    auto* data = (ShimData*)(&_data);
+
+    // What type of subsystem are we to launch into?
+    int subsystem = 0;
+    uintptr_t base = uintptr_t("") & ~0xffff;
+    auto& dos = *(const IMAGE_DOS_HEADER*)base;
+    auto& nt32 = *(const IMAGE_NT_HEADERS32*)(base + dos.e_lfanew);
+    if (nt32.FileHeader.Machine == IMAGE_FILE_MACHINE_AMD64)
+    {
+        auto& nt64 = (const IMAGE_NT_HEADERS64&)nt32;
+        subsystem = nt64.OptionalHeader.Subsystem;
+    }
+    else
+        subsystem = nt32.OptionalHeader.Subsystem;
+
+    // As this is just a shim we don't want it to interfere with Ctrl-C/Break.
+    if (subsystem == IMAGE_SUBSYSTEM_WINDOWS_CUI)
+    {
+        auto ctrlc_handler = [] (DWORD type) -> BOOL {
+            return (type == CTRL_C_EVENT || type == CTRL_BREAK_EVENT);
+        };
+        SetConsoleCtrlHandler(ctrlc_handler, TRUE);
+    }
+
+    // Work out how to treat the executable location
+    bool slashed = false;
+    for (const wchar_t* i = data->path; !slashed && *i; ++i)
+        slashed = (*i == '/');
+
+    enum style_e { style_path, style_absolute, style_relative };
+
+    style_e style = style_path;
+    if (data->path[1] == ':')
+        style = style_absolute;
+    else if (slashed)
+        style = style_relative;
+
+    // Build path to executable.
+    wchar_t* __restrict exe = data->path;
+    wchar_t* __restrict args = exe + data->path_length;
+    wchar_t* __restrict buffer = args + data->args_length;
+
+    if (style == style_relative)
+    {
+        exe = buffer;
+        wchar_t* __restrict write = exe + GetModuleFileNameW(nullptr, exe, 0x4000);
+        do { --write; } while (write[-1] != '\\');
+        const wchar_t* read = data->path;
+        do { *write++ = *read++; } while (write[-1]);
+        buffer = write;
+    }
+
+    // Build arguments.
+    wchar_t* __restrict cmd = buffer;
+    wchar_t* __restrict write = cmd;
+
+    const wchar_t* cmd_line = GetCommandLineW();
+    if (*cmd_line == '\"')
+    {
+        int i = 2;
+        do {
+            *write++ = *cmd_line++;
+        } while ((i -= cmd_line[-1] == '\"') && *cmd_line);
+        if (i)
+            *write++ = '\"';
+    }
+    else
+    {
+        do { *write++ = *cmd_line; } while (*cmd_line++ > ' ');
+        --write;
+        --cmd_line;
+    }
+
+    do { *write++ = *args; } while (*args++);
+    --write;
+    do { *write++ = *cmd_line++; } while (*cmd_line);
+
+    int ret = 0;
+    { Vt100Scope _;
+
+    // Launch the main process
+    PROCESS_INFORMATION pi;
+    STARTUPINFOW si;
+    si.cb = sizeof(si);
+    BOOL ok = CreateProcessW(exe, cmd, nullptr, nullptr, FALSE,
+        INHERIT_PARENT_AFFINITY, nullptr, nullptr, &si, &pi);
+    if (ok == TRUE)
+    {
+        if (subsystem == IMAGE_SUBSYSTEM_WINDOWS_CUI)
+        {
+            WaitForSingleObject(pi.hProcess, INFINITE);
+            DWORD process_ret = 0;
+            GetExitCodeProcess(pi.hProcess, &process_ret);
+            ret = int(process_ret);
+        }
+    }
+    else
+        ret = 9009; // cmd.exe uses this if a command doesn't exist
+
+    close_handle(pi.hThread);
+    close_handle(pi.hProcess);
+
+    } // vt100
+
+    quit(ret);
+}
+
+#endif // SHIM_BUILD
+
+// {{{1 dump -------------------------------------------------------------------
+#if SHIM_BUILD == 3
+
+static int dump(const char* tag, const char* exe_path)
+{
+    HANDLE pe_file = CreateFile(exe_path, GENERIC_READ, 0, nullptr, OPEN_EXISTING, 0, nullptr);
+    if (pe_file == INVALID_HANDLE_VALUE)
+        return 1;
+
+    void* pe_data = VirtualAlloc(nullptr, 0x40000, MEM_COMMIT, PAGE_READWRITE);
+    DWORD pe_size = 0;
+    ReadFile(pe_file, pe_data, 0x8000, &pe_size, nullptr);
+
+    auto& dos = *(IMAGE_DOS_HEADER*)pe_data;
+    auto& nt = *(IMAGE_NT_HEADERS64*)(uintptr_t(pe_data) + dos.e_lfanew);
+    if (nt.FileHeader.Machine != IMAGE_FILE_MACHINE_AMD64)
+        return 1;
+
+    IMAGE_SECTION_HEADER* shim_section = IMAGE_FIRST_SECTION(&nt);
+    auto* end_section = shim_section + nt.FileHeader.NumberOfSections;
+    for (; shim_section != end_section; ++shim_section)
+        if (*(int*)(shim_section->Name) == 'ihs.')
+            break;
+
+    if (shim_section == end_section)
+        return 1;
+
+    // This gives us a large buffer following .shim. segment when the PE's loaded
+    nt.OptionalHeader.SizeOfImage = 0x4000;
+
+    // Validate shim starts with a run of known constants
+    auto* shim_data = (uint16_t*)(uintptr_t(pe_data) + shim_section->PointerToRawData);
+    for (int i = 0; i < 16; ++i)
+    {
+        uint16_t value = shim_data[i];
+        if (value == DATA_ID)   break;
+        if (value != LENGTH_ID) return 1;
+    }
+
+    auto& offset_to = [=] (const void* addr) {
+        return int(uintptr_t(addr) - uintptr_t(pe_data));
+    };
+
+    char* buffer = (char*)(uintptr_t(pe_data) + pe_size);
+
+    HANDLE out_file = CreateFile("shim.py", GENERIC_WRITE, 0, nullptr, OPEN_ALWAYS, 0, nullptr);
+    if (out_file == INVALID_HANDLE_VALUE)
+        return 1;
+    SetFilePointer(out_file, 0, nullptr, FILE_END);
+
+    auto& print = [=] (const char* format, ...) {
+        va_list args;
+        va_start(args, format);
+        int count = wvsprintfA(buffer, format, args);
+        va_end(args);
+
+        DWORD written;
+        WriteFile(out_file, buffer, count, &written, nullptr);
+    };
+
+    print("    %s_section = %d\n", tag, offset_to(shim_section));
+    print("    %s_subsystem = %d\n", tag, offset_to(&nt.OptionalHeader.Subsystem));
+    print("    %s_alignment = %d\n", tag, nt.OptionalHeader.SectionAlignment);
+    print("    %s_data = %d\n", tag, shim_section->PointerToRawData);
+
+    print("    %s_payload = (\n", tag);
+    pe_size = shim_section->PointerToRawData;
+    for (auto* pe_read = (unsigned char*)pe_data; pe_size;)
+    {
+        print("        b\"");
+        for (int i = 0; pe_size && i < 64; ++i, --pe_size)
+            print("\\x%02x", *pe_read++);
+        print("\"\n");
+    }
+    print("    )\n");
+
+    CloseHandle(out_file);
+    return 0;
+}
+
+//------------------------------------------------------------------------------
+void __cdecl dump(void*)
+{
+    DeleteFile("shim.py");
+
+    int ret = 0;
+    ret += dump("py", "py.exe");
+    ret += dump("shim", "shim.exe");
+
+    quit(ret);
+}
+
+#endif // SHIM_BUILD
+
+// vim: foldlevel=1
diff --git a/Engine/Extras/ushell/channels/flow/core/system/lib/fsutils.py b/Engine/Extras/ushell/channels/flow/core/system/lib/fsutils.py
new file mode 100644
--- /dev/null
+++ b/Engine/Extras/ushell/channels/flow/core/system/lib/fsutils.py
@@ -0,0 +1,144 @@
+# Copyright Epic Games, Inc. All Rights Reserved.
+
+import os
+
+#-------------------------------------------------------------------------------
+def read_files(dir):
+    try:
+        files = os.scandir(dir)
+        yield from files
+        files.close()
+    except FileNotFoundError:
+        pass
+
+#-------------------------------------------------------------------------------
+def read_dirs(dir):
+    return (x for x in read_files(dir) if x.is_dir())
+
+#-------------------------------------------------------------------------------
+def get_mtime(x):
+    try: return os.path.getmtime(x)
+    except FileNotFoundError: return 0
+
+#-------------------------------------------------------------------------------
+def _rmdir_nt(dir, count=0):
+    # shutil.rmtree() deletes the contents of junctions so we can't use it.
+    try:
+        dir_attr = os.stat(dir, follow_symlinks=False).st_file_attributes
+        if dir_attr & 0x400: # 0x400 = FILE_ATTRIBUTE_REPARSE_POINT:
+            os.rmdir(dir)
+            return
+    except FileNotFoundError:
+        pass
+
+    return _rmdir_impl(dir, count)
+
+#-------------------------------------------------------------------------------
+def _rmdir_posix(dir, count=0):
+    if os.path.islink(dir):
+        os.remove(dir)
+        return
+
+    return _rmdir_impl(dir, count)
+
+#-------------------------------------------------------------------------------
+def _rmdir_impl(dir, count=0):
+    if not os.path.isdir(dir):
+        return
+
+    for item in read_files(dir):
+        if item.is_dir():
+            rmdir(item.path, count + 1)
+        else:
+            try: os.unlink(item.path)
+            except PermissionError:
+                # Try again but with the read-only flag cleared
+                os.chmod(item.path, 0xff)
+                os.unlink(item.path)
+
+    # Windows takes its time to delete a directory which can cause timing-
+    # dependent 'access denied' errors. So we'll move directories to be deleted
+    # to known and unique location prior to deletion
+    t = dir[:-1] + "~" + str(count)
+    os.rename(dir, t)
+    os.rmdir(t)
+
+#-------------------------------------------------------------------------------
+if os.name == "nt":
+    rmdir = _rmdir_nt
+else:
+    rmdir = _rmdir_posix
+
+
+
+#-------------------------------------------------------------------------------
+def import_script(script_path):
+    import importlib.util as import_util
+    spec = import_util.spec_from_file_location("", script_path)
+    module = import_util.module_from_spec(spec)
+    spec.loader.exec_module(module)
+    return module
+
+
+
+#-------------------------------------------------------------------------------
+class WorkPath(object):
+    def __init__(self, target_path):
+        self._path = target_path
+        try: os.makedirs(os.path.dirname(target_path))
+        except (FileExistsError, FileNotFoundError): pass
+
+    def __del__(self):
+        if self._path:
+            rmdir(self._path)
+
+    def __add__(self, x): return self._path + x
+    def __str__(self):    return self._path
+    def __fspath__(self): return self._path
+    def keep(self):       self._path = None
+
+
+
+#-------------------------------------------------------------------------------
+class DirLock(object):
+    def __init__(self, lock_dir):
+        self._lock_file = None
+        self._dir = os.path.abspath(lock_dir) + "/"
+        self._mtime = get_mtime(self._dir)
+
+    def __del__(self):
+        self.release()
+
+    def can_claim(self):
+        try: os.unlink(self._dir + "lock")
+        except FileNotFoundError: pass
+        except: return False
+        return True
+
+    def claim(self):
+        if self._lock_file:
+            return True
+
+        # Don't allow a claim if the dir was changed since this lock was created
+        now_mtime = get_mtime(self._dir)
+        if now_mtime > self._mtime:
+            return False
+
+        # Clean up
+        try: os.unlink(self._dir + "lock")
+        except FileNotFoundError: pass
+        except: return False
+
+        # Claim the lock
+        try: self._lock_file = open(self._dir + "lock", "wb")
+        except: pass
+
+        return True
+
+    def release(self):
+        if self._lock_file:
+            self._lock_file.close()
+            os.unlink(self._dir + "lock")
+            self._lock_file = None
+            return True
+        return False
diff --git a/Engine/Extras/ushell/channels/flow/core/system/lib/shims.py b/Engine/Extras/ushell/channels/flow/core/system/lib/shims.py
new file mode 100644
--- /dev/null
+++ b/Engine/Extras/ushell/channels/flow/core/system/lib/shims.py
@@ -0,0 +1,89 @@
+# Copyright Epic Games, Inc. All Rights Reserved.
+
+import os
+import sys
+import fsutils
+
+#-------------------------------------------------------------------------------
+class _BuilderBase(object):
+    def __init__(self, shims_dir):
+        self._out_dir = shims_dir
+
+    def commit(self):
+        self._out_dir.keep()
+
+    def create_python_shim(self, name, *args):
+        shim_path = self._out_dir + name
+        if os.path.exists(shim_path):
+            return
+
+        self._create_python_shim_impl(shim_path, *args)
+
+        # A '.' prefix was chosen to "namespace" commands in the host shell in a
+        # manner that's easy to access (e.g. typing <dot><tab>). It turns out
+        # however that Windows is picky about ".foobar". So we'll create a fallback
+        if name[0] == ".":
+            shim_path = f"{self._out_dir}_{name[1:]}"
+            self._create_python_shim_impl(shim_path, *args)
+
+    def create_shim(self, name, exe, *args):
+        shim_path = self._out_dir + name
+        if os.path.exists(shim_path):
+            return
+
+        self._create_shim_impl(shim_path, exe, *args)
+
+#-------------------------------------------------------------------------------
+class _BuilderNt(_BuilderBase):
+    def clean(self):
+        # Clean up the existing shims dir. This can fail if Windows has locked it
+        dead_shims_dir = self._out_dir + "../_deadshims/"
+        try:
+            if os.path.isdir(self._out_dir):
+                try: os.makedirs(dead_shims_dir)
+                except FileExistsError: pass
+                os.rename(self._out_dir, dead_shims_dir + str(os.getpid()))
+            fsutils.rmdir(dead_shims_dir)
+        except (FileNotFoundError, PermissionError):
+            #_log.write("Failed to clean dead shims directory")
+            pass
+
+        self._out_dir = fsutils.WorkPath(self._out_dir)
+
+    def _create_python_shim_impl(self, shim_path, *args):
+        if not shim_path.endswith(".exe"):
+            shim_path += ".exe"
+
+        import exeshim
+        exeshim.create_py(shim_path, *args)
+
+    def _create_shim_impl(self, shim_path, exe, *args):
+        if not shim_path.endswith(".exe"):
+            shim_path += ".exe"
+
+        args = ('"' + '" "'.join(args) + '"') if args else ""
+        import exeshim
+        exeshim.create(shim_path, exe, args)
+
+#-------------------------------------------------------------------------------
+class _BuilderPosix(_BuilderBase):
+    def clean(self):
+        fsutils.rmdir(self._out_dir)
+        self._out_dir = fsutils.WorkPath(self._out_dir)
+
+    def _create_python_shim_impl(self, shim_path, *args):
+        exe = os.path.abspath(sys.executable)
+        return self._create_shim_impl(shim_path, exe, *args)
+
+    def _create_shim_impl(self, shim_path, exe, *args):
+        import shlex
+        with open(shim_path, "wt") as out:
+            out.write("#!/bin/sh\n")
+            items = (shlex.quote(x) for x in (exe, *args))
+            out.write(" ".join(items))
+            out.write(" \"$@\"\n")
+
+        import stat
+        os.chmod(shim_path, stat.S_IRWXU)
+
+Builder = _BuilderNt if os.name == "nt" else _BuilderPosix
diff --git a/Engine/Extras/ushell/channels/flow/core/system/run.py b/Engine/Extras/ushell/channels/flow/core/system/run.py
new file mode 100644
--- /dev/null
+++ b/Engine/Extras/ushell/channels/flow/core/system/run.py
@@ -0,0 +1,76 @@
+# Copyright Epic Games, Inc. All Rights Reserved.
+
+import os
+import sys
+
+#-------------------------------------------------------------------------------
+def _print_sub_commands(*args):
+    from flow.text import text
+    with text.light_cyan:
+        _print_sub_commands_impl(*args)
+
+#-------------------------------------------------------------------------------
+def _print_sub_commands_impl(tree_node):
+    import shutil
+    max_width = shutil.get_terminal_size(fallback=(80,0))[0]
+
+    once = True
+    for child_name, child_node in tree_node.read_children():
+        if once:
+            print("\nSUB-COMMANDS:")
+            once = False
+
+        # Get a description for the sub-command
+        if child_class := child_node.get_command_class():
+            child_command = child_class.construct()
+            desc = child_command.get_desc().lstrip()
+        else:
+            desc = "[%s]" % ",".join(x for x,_ in child_node.read_children())
+        desc = desc.replace("\n", "")
+        desc = desc.replace("    ", " ")
+
+        if len(desc) > max_width - 16:
+            desc = desc[:max_width - 16 - 3] + "..."
+        print("  %-12s" % child_name, desc)
+
+#-------------------------------------------------------------------------------
+def main(args):
+    # Imperfectly undo state that leaks from the command shims on Windows
+    if sys.executable == os.getenv("PYTHONEXECUTABLE"):
+        del os.environ["PYTHONEXECUTABLE"]
+
+    # Import all the flow goodies
+    sys.path.append(os.path.abspath(__file__ + "/../"))
+    import flow.system
+
+    manifest, *args = args
+
+    # Traverse the command tree using the given arguments
+    working_dir = os.path.dirname(manifest) + "/"
+    system = flow.system.System(working_dir)
+    cmd_tree = system.get_command_tree()
+    tree_node, tail_args = cmd_tree.find(iter(args))
+
+    # Does the found tree node have a command we can execute?
+    command_class = tree_node.get_command_class()
+    if command_class:
+        command = command_class.construct()
+        return command.invoke(tail_args)
+
+    # We've reached a point in the command tree that has no command. Inform the
+    # user about the available sub-commands.
+
+    next_arg = next(tail_args, None)
+    if next_arg and next_arg != "--help":
+        from flow.text import text
+        with text.light_red:
+            print("ERRROR: Unrecognized arguments:", next_arg, *tail_args)
+
+    _print_sub_commands(tree_node)
+
+
+
+if __name__ == "__main__":
+    ret = main(sys.argv[1:])
+    ret = not ret if isinstance(ret, bool) else ret
+    raise SystemExit(ret)
diff --git a/Engine/Extras/ushell/channels/flow/core/system/version b/Engine/Extras/ushell/channels/flow/core/system/version
new file mode 100644
--- /dev/null
+++ b/Engine/Extras/ushell/channels/flow/core/system/version
@@ -0,0 +1 @@
+20
diff --git a/Engine/Extras/ushell/channels/flow/nt/boot.bat b/Engine/Extras/ushell/channels/flow/nt/boot.bat
new file mode 100755
--- /dev/null
+++ b/Engine/Extras/ushell/channels/flow/nt/boot.bat
@@ -0,0 +1,61 @@
+:: Copyright Epic Games, Inc. All Rights Reserved.
+
+@echo off
+setlocal
+chcp 65001 1>nul 2>nul
+
+:: Set up the working directory flow can store its necessities in, and the dir
+:: that contains the channels that implement the shell.
+set _working="%localappdata%\ushell\.working"
+set _channels="%~dp0..\..\..\channels"
+if defined flow_working_dir (
+    set _working="%flow_working_dir%"
+)
+
+:: Call out to user hooks
+if exist "%userprofile%\.ushell\hooks\boot.bat" (
+    call "%userprofile%\.ushell\hooks\boot.bat"
+)
+
+:: Make it easier for users to add their own channels
+set _channels=%_channels% "%userprofile%\.ushell\channels"
+if defined flow_channels_dir (
+    set _channels=%_channels% %flow_channels_dir%
+)
+
+:: Provision Python. We shim through a cmd.exe for simpler error handling in
+:: provision.bat (see the script's use of call:on_error).
+cmd.exe /d/c ""%~dp0\provision.bat" %_working%"
+if errorlevel 1 exit /b 1
+
+:: Boot up
+set _cookie=%time::=%
+set _cookie=%_cookie:~1%
+set _cookie=%temp%\ushell\cmd_boot_%_cookie%
+%_working%\python\current\flow_python.exe -Xutf8 -Esu "%~dp0..\core\system\boot.py" %_working% %_channels% -- "--bootarg=cmd,%_cookie%" %*
+
+:: was the python binary missing? "3" comes from cmd.exe - lets hope that it
+:: doesn't collide with something boot.py might return.
+if errorlevel 3 (
+    echo,
+    echo,
+    echo ERROR: Missing flow_python.exe from '%_working%\python\current'
+    timeout /T 15
+    exit 1
+)
+
+:: did boot.py fail?
+if errorlevel 1 (
+    echo,
+    echo,
+    echo ERROR: boot.py failed [%errorlevel%]
+    timeout /T 15
+    exit 1
+)
+
+endlocal & if exist "%_cookie%" (
+    for /f "usebackq delims=" %%d in ("%_cookie%") do (
+        %%d
+    )
+    del /q "%_cookie%"
+)
diff --git a/Engine/Extras/ushell/channels/flow/nt/boot.ps1 b/Engine/Extras/ushell/channels/flow/nt/boot.ps1
new file mode 100644
--- /dev/null
+++ b/Engine/Extras/ushell/channels/flow/nt/boot.ps1
@@ -0,0 +1,51 @@
+Write-Debug "Ushell bootstrap args 2: $Args"
+
+$Working = Join-Path $env:LOCALAPPDATA "ushell" ".working"
+if(![String]::IsNullOrEmpty($env:flow_working_dir)) {
+    $Working = $env:flow_working_dir
+}
+
+Write-Debug "Working dir: $Working"
+
+$Channels = @(Join-Path $PSScriptRoot "..\..\..\channels")
+$Channels += Join-Path $env:userprofile "\.ushell\channels"
+$Channels += $env:flow_channels_dir
+
+Write-Debug "Channels: $Channels"
+
+$ProvisionBat = Join-Path $PSScriptRoot "provision.bat"
+Write-Debug "Provisioning: $ProvisionBat"
+& "cmd.exe" "/d/c" "`"$ProvisionBat`" $Working"
+
+if(!$?) {
+    throw "Ushell provisioning failed"
+}
+
+
+$PythonPath = Join-Path $Working "python\current\flow_python.exe"
+$Bootpy = Join-Path $PSScriptRoot "..\core\system\boot.py"
+Write-Debug "Python: $PythonPath"
+
+$TempDir = Join-Path Temp: "ushell" 
+New-Item -ItemType Directory -Force $TempDir
+$Cookie = Join-Path $TempDir "cmd_boot_$(New-Guid).ps1"
+$Cookie = New-Item $Cookie
+Write-Debug "Cookie: $Cookie"
+
+& $PythonPath "-Xutf8" "-Esu" $Bootpy $Working @Channels "--" "--bootarg=pwsh,$Cookie" $Args | Out-Host
+$PythonReturn = $?
+
+if(!$PythonReturn) {
+    throw "Python boot failed"
+}
+
+try {
+    Import-Module $Cookie
+}
+finally {
+    if( Test-Path $Cookie )
+    {
+        Write-Debug "Removing cookie: $Cookie"
+        Remove-Item $Cookie
+    }
+}
\ No newline at end of file
diff --git a/Engine/Extras/ushell/channels/flow/nt/boot.sh b/Engine/Extras/ushell/channels/flow/nt/boot.sh
new file mode 100644
--- /dev/null
+++ b/Engine/Extras/ushell/channels/flow/nt/boot.sh
@@ -0,0 +1,23 @@
+#!/bin/bash
+
+# Copyright Epic Games, Inc. All Rights Reserved.
+
+self_dir=$(cygpath --windows --absolute $(dirname $0))
+working_dir=$LOCALAPPDATA\\ushell\\.working
+
+# Provision Python
+cmd.exe /d/c "$self_dir/provision.bat" "$working_dir"
+if [ ! $? ]; then
+    exit 1
+fi
+
+declare -a channels_dir
+channels_dir+=("$self_dir/../..")
+channels_dir+=("$USERPROFILE/.ushell/channels")
+channels_dir+=("$FLOW_CHANNELS_DIR")
+
+"$working_dir/python/current/flow_python.exe" -Esu "$self_dir/../core/system/boot.py" "$working_dir" "${channels_dir[@]}" -- $@
+if [ ! $? ]; then
+    echo boot.py failed
+    exit 1
+fi
diff --git a/Engine/Extras/ushell/channels/flow/nt/cmd_thefuzz.bat b/Engine/Extras/ushell/channels/flow/nt/cmd_thefuzz.bat
new file mode 100755
--- /dev/null
+++ b/Engine/Extras/ushell/channels/flow/nt/cmd_thefuzz.bat
@@ -0,0 +1,64 @@
+:: Copyright Epic Games, Inc. All Rights Reserved.
+
+@echo off
+
+setlocal
+set _bat="%temp%\thefuzz_%flow_sid%.bat"
+set _fzf=fzf --layout=reverse --color=16 --height=20
+
+<nul >%_bat% (
+    if "%~1"=="history" call:history
+    if "%~1"=="chdir"   call:chdir %2
+    if "%~1"=="explore" call:explore %2
+)
+
+endlocal & if not errorlevel 1 (
+    call %_bat%
+)
+
+goto:eof
+
+::------------------------------------------------------------------------------
+:history
+set _hist=
+for /f "tokens=1,2 delims==" %%d in ('doskey /macros') do (
+    if "%%d"=="history" (
+        set _hist=%%e
+        goto:_hist_set_break
+    )
+)
+:_hist_set_break
+if defined _hist (
+    call:history_impl %_hist%
+)
+goto:eof
+
+:history_impl
+for /f "delims=" %%d in ('""%~f1" %2 | rg "^^[0-9 ]+^(.+^)$" -r $1 | %_fzf% --scheme=history --history=%temp%/thefuzz_hist"') do (
+    echo set _hist_result="%%d"
+    echo set _hist_result=%%_hist_result:^|=^^^|%%
+    echo set _hist_result=%%_hist_result:^&=^^^&%%
+    echo set _hist_result=%%_hist_result:^<=^^^<%%
+    echo set _hist_result=%%_hist_result:^>=^^^>%%
+    echo echo %%_hist_result:~1,-1%%^>%temp%\thefuzz_%flow_sid%_hist_tmp
+    echo clip ^<%temp%\thefuzz_%flow_sid%_hist_tmp
+    echo echo Copied to clipboard; [92m%%_hist_result:~1,-1%%[0m
+    echo set _hist_result=
+)
+goto:eof
+
+::------------------------------------------------------------------------------
+:chdir
+set _base="%~f1"
+if "%~1"=="" set _base=\
+echo cd /d ^^
+fd -d8 -td . %_base:\"="% | %_fzf% --scheme=default --history=%temp%/thefuzz_chdir
+goto:eof
+
+::------------------------------------------------------------------------------
+:explore
+set _base="%~f1"
+if "%~1"=="" set _base=\
+echo explorer.exe ^^
+fd -d8 -td . %_base:\"="% | %_fzf% --scheme=default --history=%temp%/thefuzz_chdir
+goto:eof
diff --git a/Engine/Extras/ushell/channels/flow/nt/describe.flow.py b/Engine/Extras/ushell/channels/flow/nt/describe.flow.py
new file mode 100644
--- /dev/null
+++ b/Engine/Extras/ushell/channels/flow/nt/describe.flow.py
@@ -0,0 +1,47 @@
+# Copyright Epic Games, Inc. All Rights Reserved.
+
+import os
+import flow.describe
+
+#-------------------------------------------------------------------------------
+clink = flow.describe.Tool()
+clink.version("1.0.0a4")
+clink.platform("win32")
+clink.payload("[REDACTED_PROFILE_URL]/clink/releases/download/1.0.0a4/clink-1.0.0a4.zip")
+clink.bin("clink_x64.exe")
+clink.sha1("6a642bc48e21a3093698067f0fe09eff3ccdfaa5")
+
+#-------------------------------------------------------------------------------
+fd = flow.describe.Tool()
+fd.version("8.4.0")
+fd.platform("win32")
+fd.payload("[REDACTED_PROFILE_URL]/fd/releases/download/v8.4.0/fd-v8.4.0-x86_64-pc-windows-msvc.zip")
+fd.bin("fd.exe")
+fd.sha1("aef6f9809b7d86705fe828689453aae61f71365f")
+fd.source("[REDACTED_PROFILE_URL]/fd/releases/latest", "fd-v([0-9.]+)-x86_64")
+
+#-------------------------------------------------------------------------------
+fzf_win32 = flow.describe.Tool()
+fzf_win32.version("0.33.0")
+fzf_win32.platform("win32")
+fzf_win32.payload("[REDACTED_PROFILE_URL]/fzf/releases/download/0.33.0/fzf-0.33.0-windows_amd64.zip")
+fzf_win32.bin("fzf.exe")
+fzf_win32.sha1("be7fc7b271069f7da5c27d570c81287d23cd328e")
+fzf_win32.source("[REDACTED_PROFILE_URL]/fzf/releases/latest", "fzf-(\d+\.\d+\.\d+)-windows")
+
+#-------------------------------------------------------------------------------
+shell_cmd = flow.describe.Command()
+shell_cmd.source("shell_cmd.py", "Cmd")
+shell_cmd.invoke("boot")
+shell_cmd.prefix("$")
+
+#-------------------------------------------------------------------------------
+shell_pwsh = flow.describe.Command()
+shell_pwsh.source("shell_pwsh.py", "Pwsh")
+shell_pwsh.invoke("boot")
+shell_pwsh.prefix("$")
+
+#-------------------------------------------------------------------------------
+channel = flow.describe.Channel()
+channel.version("1")
+channel.parent("flow.core")
diff --git a/Engine/Extras/ushell/channels/flow/nt/provision.bat b/Engine/Extras/ushell/channels/flow/nt/provision.bat
new file mode 100755
--- /dev/null
+++ b/Engine/Extras/ushell/channels/flow/nt/provision.bat
@@ -0,0 +1,156 @@
+:: Copyright Epic Games, Inc. All Rights Reserved.
+
+@echo off
+call:main "%~1\python"
+goto:eof
+
+::------------------------------------------------------------------------------
+:main
+
+call:check_bin tar.exe _tar_path
+call:check_bin curl.exe _curl_path
+call:check_bin certutil.exe _certutil_path
+
+set _pyver=3.11.4
+set _pysha=126802ff9fe787b961ae4d877262e6d6ce807d636295ef429c02e4dfd3e48041
+set _pytag=311
+set _destdir=%~f1\%_pyver%
+
+if exist "%_destdir%" (
+    goto:eof
+)
+
+1>nul 2>nul (
+    rd "%_destdir%\..\current"
+)
+
+call:get_python "%_destdir%"
+
+1>nul 2>nul (
+    mklink /j "%_destdir%\..\current" "%_destdir%"
+)
+
+goto:eof
+
+::------------------------------------------------------------------------------
+:check_bin
+if exist "%SYSTEMROOT%\System32\%~1" (
+    set %~2="%SYSTEMROOT%\System32\%~1"
+    goto:check_bin_break
+) else 2>nul (
+    for /f "delims=" %%d in ('where %~1') do (
+        set %~2="%SYSTEMROOT%\System32\%%~d"
+        goto:check_bin_break
+    )
+)
+set errorlevel=1
+:check_bin_break
+call:on_error "Win10 1803+ required - unable to find binary %~1"
+goto:eof
+
+::------------------------------------------------------------------------------
+:get_python
+set _workdir=%~1~
+for /f "delims=, tokens=2" %%d in ('tasklist.exe /fo csv /fi "imagename eq cmd.exe" /nh') do (
+    set _workdir=%~1_%%~d
+)
+
+1>nul 2>nul (
+    mkdir "%_workdir%"
+)
+
+1>"%_workdir%\provision.log" 2>&1 (
+    set prompt=-$s
+    echo on
+    call:get_python_impl "%_workdir%"
+    echo off
+)
+
+move "%_workdir%" "%~1" 1>nul 2>nul
+if not exist "%~1" (
+    call:on_error "Failed finalising Python provision"
+)
+
+echo.
+goto:eof
+
+::------------------------------------------------------------------------------
+:get_python_impl
+rd /q /s "%~1" 2>nul
+if exist "%~1" call:on_error "Unable to remove directory '%~1'"
+md "%~1"
+pushd "%~1"
+
+call:echo "Working path; %~1"
+call:echo "1/2 : Getting Python %_pyver%"
+call:get_url https://www.python.org/ftp/python/%_pyver%/python-%_pyver%-embed-amd64.zip
+call:integrity python-%_pyver%-embed-amd64.zip %_pysha%
+call:unzip python-%_pyver%-embed-amd64.zip .
+
+del *._pth
+call:unzip python%_pytag%.zip Lib
+
+md DLLs
+move *.pyd DLLs
+move lib*.dll DLLs
+move sq*.dll DLLs
+
+:: One particular AV provider began intervening with the python.exe process that
+:: underpins a shell. Renaming it seems to be a successful workaround.
+copy python.exe flow_python.exe
+
+call:echo "2/2 : Adding Pip"
+call:get_url https://bootstrap.pypa.io/get-pip.py
+.\flow_python.exe get-pip.py
+call:on_error "Failed to add Pip"
+del get-pip.py
+
+popd
+goto:eof
+
+::------------------------------------------------------------------------------
+:get_url
+%_curl_path% --output "%~nx1" "%~1"
+call:on_error "Failed to get url '%~1'"
+goto:eof
+
+::------------------------------------------------------------------------------
+:integrity
+%_certutil_path% -hashfile "%~1" sha256 | findstr "%~2"
+if errorlevel 1 (
+    %_certutil_path% -hashfile "%~1" sha256
+    set errorlevel=1
+    call:on_error "Integrity check failed for '%~1'"
+)
+goto:eof
+
+::------------------------------------------------------------------------------
+:unzip
+mkdir "%~2"
+%_tar_path% -xf "%~1" -C "%~2"
+call:on_error "Failed to unzip '%~1'"
+del "%~1"
+goto:eof
+
+::------------------------------------------------------------------------------
+:echo
+1>con echo %~1
+goto:eof
+
+::------------------------------------------------------------------------------
+:on_error
+if not %errorlevel%==0 1>&2 (
+    echo ERROR: %~1
+    1>con (
+        :: timeout's printing is buggy when 1>con is used :(
+        echo.
+        echo.
+        echo.
+        echo.
+        echo.
+        echo ERROR: %~1
+        timeout /T 5
+    )
+    exit 1
+)
+goto:eof
diff --git a/Engine/Extras/ushell/channels/flow/nt/shell_cmd.py b/Engine/Extras/ushell/channels/flow/nt/shell_cmd.py
new file mode 100644
--- /dev/null
+++ b/Engine/Extras/ushell/channels/flow/nt/shell_cmd.py
@@ -0,0 +1,190 @@
+# Copyright Epic Games, Inc. All Rights Reserved.
+
+import os
+import sys
+import flow.cmd
+from pathlib import Path
+
+#-------------------------------------------------------------------------------
+class Cmd(object):
+    def _add_clink_lua(self):
+        channel = self.get_channel()
+        system = channel.get_system()
+        working_dir = system.get_working_dir()
+
+        lua_path = working_dir + "complete.clink.lua"
+        if not system.is_path_stale(lua_path):
+            return
+
+        cmd_tree = system.get_command_tree()
+        tree_root = cmd_tree.get_root_node()
+
+        lua = open(lua_path, "w")
+        lua.write("local commands = {")
+        for name,_ in tree_root.read_children():
+            lua.write(f"'{name}',")
+        lua.write("}\n")
+
+        run_py_path = os.path.abspath(__file__ + "/../../core/system/run.py")
+        cmd = rf'"{sys.executable}" -Xutf8 -Esu "{run_py_path}" "{working_dir}/manifest" $complete --daemon'
+        cmd = cmd.replace("\\", "/")
+        lua.write(f"local py_server_cmd = '{cmd}'")
+
+        lua.write(_clink_lua)
+        lua.close()
+
+    def run(self, env):
+        # Give the user a convenient way to add customisations in to the session
+        self.print_info("Customisation")
+        home_dir = Path(self.get_home_dir())
+        tags = (
+            flow.cmd.text.grey("none"),
+            flow.cmd.text.white("okay"),
+        )
+
+        # These are really a no-op. It is here to hint to the user how they can
+        # add their own channels
+        candidate = home_dir / "channels"
+        found = candidate.is_dir()
+        print(tags[int(found)], candidate / "*")
+
+        # Indicate to the user that there's a way to hook the boot process too
+        candidate = home_dir / "hooks/boot.bat"
+        found = candidate.is_file()
+        print(tags[int(found)], candidate)
+
+        # This is run just before the prompt is shown
+        candidate = home_dir / "hooks/startup.bat"
+        found = candidate.is_file()
+        print(tags[int(found)], candidate)
+        if found:
+            self._user_script = candidate
+
+        print()
+        super().run(env)
+
+    def register_shells(self, registrar):
+        registrar.add("cmd", lambda x: self)
+        return super().register_shells(registrar)
+
+    def boot_shell(self, env, cookie):
+        self._add_clink_lua()
+
+        channel = self.get_channel()
+        system = channel.get_system()
+        working_dir = system.get_working_dir()
+        env["CLINK_PATH"] = working_dir
+        env["DIRCMD"] = "/ogen"
+
+        if prompt := env.get("FLOW_PROMPT", None):
+            prompt = prompt.replace("\n", "\r")
+            env["FLOW_PROMPT"] = prompt
+
+        try: os.makedirs(os.path.dirname(cookie))
+        except: pass
+
+        thefuzz_path = Path(__file__).parent / "cmd_thefuzz.bat"
+        thefuzz_path = thefuzz_path.resolve()
+
+        with open(cookie, "wt") as out:
+            print(f"cd /d \"{os.getcwd()}\"", file=out)
+            for key, value in env.read_changes():
+                value = value or ""
+                print(fr'set "{key}={value}"', file=out)
+
+            print(rf'doskey ff=fd $B fzf $B clip', file=out)
+            print(rf'doskey qq="{thefuzz_path}" history', file=out)
+            print(rf'doskey de="{thefuzz_path}" explore $*', file=out)
+            print(rf'doskey dc="{thefuzz_path}" chdir $*', file=out)
+            print(rf'doskey dc.="{thefuzz_path}" chdir .', file=out)
+
+            print("clink_x64.exe inject", file=out)
+            print("$tip.exe", file=out)
+
+            if user_script := getattr(self, "_user_script", None):
+                print(rf'call "{user_script}"', file=out)
+
+
+#-------------------------------------------------------------------------------
+_clink_lua = r"""
+local py_server = io.popen2(py_server_cmd)
+
+for _, tree_root in ipairs(commands) do
+    local arg_handler = function(index, line)
+        py_server:write("\x01\n")
+        py_server:write(tree_root.."\n")
+        local n = line:getwordcount()
+        for i = 2, n do
+            py_server:write(line:getword(i).."\n")
+        end
+
+        local delim = line:getwordinfo(n).delim
+        if delim ~= " " then
+            py_server:write("\x03"..delim.."\n")
+        end
+        py_server:write("\x02\n")
+
+        local ret = {}
+        while true do
+            match = py_server:read()
+            if match == nil then return end
+            if match == "\x01" then return ret end
+            if match == "\x02" then return end
+            table.insert(ret, match)
+        end
+    end
+
+    clink.argmatcher(tree_root)
+        :loop(1)
+        :setflagprefix("-")
+        :addarg(arg_handler)
+        :addflags(arg_handler)
+end
+
+local cmd_generator = clink.generator(0)
+function cmd_generator:generate(line_state, match_builder)
+    if line_state:getwordcount() ~= 1 then return end
+    if line_state:getword(1) ~= "." then return end
+
+    match_builder:setprefixincluded()
+    for _, command in ipairs(commands) do
+        if command:sub(1, 1) == "." then
+            match_builder:addmatch(command)
+        end
+    end
+    return true
+end
+
+function cmd_generator:getprefixlength(line_state)
+    if line_state:getwordcount() ~= 1 then return end
+    local prefix = line_state:getword(1):sub(1, 1)
+    if prefix == "." then return 1 end
+    return 0
+end
+
+local cliff_prompt = clink.promptfilter(-493)
+function cliff_prompt:filter(prompt)
+    prompt_override = os.getenv("FLOW_PROMPT")
+    if prompt_override == nil then
+        return prompt
+    end
+
+    py_server:write("\x01\n")
+    py_server:write("$\n")
+    py_server:write(prompt_override)
+    py_server:write("\n\x02\n")
+
+    prompt = nil
+    while true do
+        ret = py_server:read()
+        if ret == nil then return end
+        if ret == "\x01" then return prompt end
+        if ret == "\x02" then return end
+        if not prompt then
+            prompt = ret
+        else
+            prompt = prompt.."\n"..ret
+        end
+    end
+end
+"""
