Updated ushell [26875160]

diff --git a/Engine/Extras/ushell/channels/flow/nt/shell_pwsh.py b/Engine/Extras/ushell/channels/flow/nt/shell_pwsh.py
new file mode 100644
--- /dev/null
+++ b/Engine/Extras/ushell/channels/flow/nt/shell_pwsh.py
@@ -0,0 +1,215 @@
+# Copyright Epic Games, Inc. All Rights Reserved.
+
+import os
+import sys
+
+#-------------------------------------------------------------------------------
+class Pwsh(object):
+    def register_shells(self, registrar):
+        registrar.add("pwsh", _Shell)
+        return super().register_shells(registrar)
+
+class _Shell(object):
+    def __init__(self, system):
+        self._system = system
+
+    def get_system(self):
+        return self._system
+
+    def boot_shell(self, env, cookie):
+        # no prompt modification here for now
+        # could try and modify the Prompt function from the cookie
+        # could also publish environment variables users can use from their own prompts
+        try: os.makedirs(os.path.dirname(cookie))
+        except: pass
+
+        system = self.get_system();
+        working_dir = system.get_working_dir()
+        shims_path = os.path.normpath(working_dir) + "/shims"
+        cmd_tree = system.get_command_tree()
+        tree_root = cmd_tree.get_root_node()
+        run_py_path = os.path.abspath(__file__ + "/../../core/system/run.py")
+        manifest = f"{working_dir}/manifest"
+        start_dir = os.getcwd()
+
+        with open(cookie, "wt") as out:
+            out.write(f"Set-Location \"{start_dir}\"\n")
+            # Environment must be written out first so autocomplete daemon gets the right session id
+            for key, value in env.read_changes():
+                value = value or ""
+                out.write(f'${{env:{key}}}="{value}"\n')
+            out.write(_header.format(shims_path=shims_path,working_dir=working_dir))
+            out.write(_autocompleter.format(python=sys.executable, run_py_path=run_py_path, working_dir=working_dir))
+            for name,_ in tree_root.read_children():
+                if name.startswith("$"): continue
+                out.write(_function_template.format(name=name))
+            out.write(_cleanup)
+
+_header = r"""
+$ShimsPath = Resolve-Path "{shims_path}"
+$env:Path += ";$ShimsPath"
+"""
+
+# Braces in this string must be doubled to escape formatting
+#TODO: give these functions long names, then optionally add aliases?
+_function_template = r"""
+Register-ArgumentCompleter -CommandName "{name}.exe" -Native -ScriptBlock $SharedCompleter
+Register-ArgumentCompleter -CommandName "{name}" -Native -ScriptBlock $SharedCompleter
+# Explicitly export this function so that others are not exported automatically
+Export-ModuleMember -Function {name}
+"""
+# If we give commands long names, we need to separate completers that provide the name ushell $complete expects
+_autocompleter = r"""
+function Get-Daemon {{
+    if( $script:CompleteDaemonCache -ne $null -and !$script:CompleteDaemonCache.HasExited) {{
+        return $script:CompleteDaemonCache
+    }}
+
+    $ProcStart = New-Object System.Diagnostics.ProcessStartInfo
+    # the $ in $complete must be escaped for powershell
+    $DaemonExe = $true
+    if( $DaemonExe ) {{
+        $ProcStart.FileName = Join-Path $ShimsPath "`$complete"
+        $ProcStart.Arguments = @("--daemon")
+    }}
+    else {{
+        $ProcStart.Filename = "{python}"
+        $ProcStart.Arguments = @(
+            "-Xutf8",
+            "-Esu",
+            "{run_py_path}",
+            (Join-Path {working_dir} "manifest"),
+            "`$complete",
+            "--daemon"
+            )
+    }}
+    $ProcStart.UseShellExecute = $false
+    $ProcStart.RedirectStandardOutput = $true
+    $ProcStart.RedirectStandardInput = $true
+
+    $script:CompleteDaemonCache = [System.Diagnostics.Process]::Start($ProcStart)
+    if( $null -eq $script:CompleteDaemonCache ) {{
+        throw "Failed to start autocomplete Daemon!"
+    }}
+    return $script:CompleteDaemonCache
+}}
+
+$SharedCompleter = {{
+    param($WordToComplete, $CommandAst, $CursorPosition)
+
+    $Words = -split "$CommandAst"
+    $Words[0] = $Words[0] -replace "\.exe"
+
+    # Add ... if we have a partial argument
+    if(![string]::IsNullOrEmpty($WordToComplete)) {{
+        $Words[-1] += "..."
+    }}
+    else {{
+        $Words += ""
+    }}
+
+    $CompleteDaemon = Get-Daemon
+
+    $CompleteDaemon.StandardInput.WriteLine("`u{{01}}")
+    foreach( $Word in $Words ) {{
+        $CompleteDaemon.StandardInput.WriteLine($Word)
+    }}
+    $CompleteDaemon.StandardInput.WriteLine("`u{{02}}")
+
+    $Results = @()
+    $Continue = $true
+    while($Continue) {{
+        $Option = $CompleteDaemon.StandardOutput.Readline();
+        # -eq and .Equals behave differently for the ascii separator bytes used here
+        switch -exact ($Option) {{
+            {{ "`u{{01}}".Equals($Option) }} {{
+                $Continue = $false
+                break;
+            }}
+            {{ "`u{{02}}".Equals($Option) }} {{
+                $Results = @();
+                $Continue = $false
+                break;
+            }}
+            "" {{
+                $Continue = $false
+                break;
+            }}
+            default {{
+                $Results += $Option
+                break;
+            }}
+        }}
+    }}
+
+    # Filter to only arguments that start with our partial match if it exists
+    $Results = $Results | Where-Object {{ $_ -like "$WordToComplete*" }}
+
+    if( $results.Count -eq 0 ) {{
+        # prevent powershell from autocompleting paths to match ushell behavior in cmd
+        ""
+    }}
+    else {{
+        $results
+    }}
+}}
+
+$EnvVarsToRemove = [System.Collections.Generic.HashSet[string]]@()
+
+function Update-UShellEnvVars {{
+    $Prompt = $env:FLOW_PROMPT
+    if ([string]::IsNullOrEmpty($Prompt)) {{
+        return
+    }}
+
+    $CompleteDaemon = Get-Daemon
+
+    if( $null -eq $CompleteDaemon ) {{
+        Write-Error "Update-UShellEnvVars - CompleteDaemon is null!"
+    }}
+    if( $CompleteDaemon.HasExited ) {{
+        Write-Error "Update-UShellEnvVars - CompleteDaemon has exited!"
+    }}
+
+    $CompleteDaemon.StandardInput.WriteLine("`u{{01}}")
+    $CompleteDaemon.StandardInput.WriteLine("`$`$")
+    $CompleteDaemon.StandardInput.WriteLine("`u{{02}}")
+
+    $Results = @{{}}
+    $Continue = $true
+    while($Continue) {{
+        $Key = $CompleteDaemon.StandardOutput.Readline();
+        # -eq and .Equals behave differently for the ascii separator bytes used here
+        if( [string]::IsNullOrEmpty($Key) -or $Key.Equals("`u{{01}}") ) {{ break; }}                 #EOF
+        if( $Key.Equals("`u{{02}}") ) {{ $Results = @{{}}; break; }} #error
+
+        $Value = $CompleteDaemon.StandardOutput.Readline();
+        if( [string]::IsNullOrEmpty($Value) -or $Value.Equals("`u{{01}}") -or $Value.Equals("`u{{02}}") ) {{ $Results = @{{}}; break; }} #error
+
+        $Results[$Key] = $Value
+    }}
+
+    foreach( $i in $Results.GetEnumerator()) {{
+        $EnvPath = (Join-Path Env: "USHELL_$($i.Key)")
+        Write-Debug "Setting env var: $EnvPath"
+        $null = $script:EnvVarsToRemove.Add($EnvPath)
+        New-Item $EnvPath -Value $i.Value -Force
+    }}
+}}
+
+Export-ModuleMember -Function "Update-UShellEnvVars"
+"""
+
+# Kill the autocomplete daemon so we don't leak a process
+_cleanup = r"""
+$ExecutionContext.SessionState.Module.OnRemove += {
+    if( $script:CompleteDaemonCache -ne $null ) {
+        Write-Debug "Stopping complete daemon"
+        Stop-Process $script:CompleteDaemonCache
+    }
+    foreach( $v in $script:EnvVarsToRemove.GetEnumerator()) {
+        Write-Debug "Removing env var: $v"
+        Remove-Item $v
+    }
+}
+"""
\ No newline at end of file
diff --git a/Engine/Extras/ushell/channels/flow/posix/boot.sh b/Engine/Extras/ushell/channels/flow/posix/boot.sh
new file mode 100644
--- /dev/null
+++ b/Engine/Extras/ushell/channels/flow/posix/boot.sh
@@ -0,0 +1,25 @@
+#!/bin/bash
+
+# Copyright Epic Games, Inc. All Rights Reserved.
+
+self_dir=$(dirname $0)
+working_dir=~/.ushell/.working
+
+# provision python >=3.8 using provision.sh
+if ! $self_dir/provision.sh $working_dir; then
+    echo Failed to provision Python
+    exit 1
+fi
+
+declare -a channels_dir
+channels_dir+=("$self_dir/../..")
+channels_dir+=("$HOME/.ushell/channels")
+if [[ -n "$FLOW_CHANNELS_DIR" ]]; then
+    channels_dir+=("$FLOW_CHANNELS_DIR")
+fi
+
+"$working_dir/python/current/bin/python" -Esu "$self_dir/../core/system/boot.py" "$working_dir" "${channels_dir[@]}" -- $@
+if [ ! $? ]; then
+    echo boot.py failed
+    exit 1
+fi
diff --git a/Engine/Extras/ushell/channels/flow/posix/describe.flow.py b/Engine/Extras/ushell/channels/flow/posix/describe.flow.py
new file mode 100644
--- /dev/null
+++ b/Engine/Extras/ushell/channels/flow/posix/describe.flow.py
@@ -0,0 +1,14 @@
+# Copyright Epic Games, Inc. All Rights Reserved.
+
+import flow.describe
+
+#-------------------------------------------------------------------------------
+shells = flow.describe.Command()
+shells.source("shells.py", "Shells")
+shells.invoke("boot")
+shells.prefix("$")
+
+#-------------------------------------------------------------------------------
+channel = flow.describe.Channel()
+channel.version("1")
+channel.parent("flow.core")
diff --git a/Engine/Extras/ushell/channels/flow/posix/provision.sh b/Engine/Extras/ushell/channels/flow/posix/provision.sh
new file mode 100644
--- /dev/null
+++ b/Engine/Extras/ushell/channels/flow/posix/provision.sh
@@ -0,0 +1,102 @@
+#!/bin/bash
+
+# Copyright Epic Games, Inc. All Rights Reserved.
+
+_error()   { printf "\x1b[91m!! ERROR: %s\x1b[0m\n" "$1" ; }
+_header()  { printf "\x1b[96m-- %s\x1b[0m\n" "$1" ; }
+_success() { printf "\x1b[92m%s\x1b[0m\n" "$1" ; }
+
+py_ver=3.10
+
+# Try and find a suitable Python binary
+_header "Looking for a Python binary"
+if command -v python${py_ver}; then
+    py_bin=python${py_ver}
+elif command -v /usr/local/opt/python@${py_ver}/bin/python3; then
+    py_bin=/usr/local/opt/python@${py_ver}/bin/python3
+elif command -v python3; then
+    py_bin=python3
+else
+    echo FAILED
+    exit 1
+fi
+_success Found
+
+# Work out what version of Python we've got
+py_ver_full=$($py_bin --version)
+IFS=" "
+read -ra py_ver_full <<< "$py_ver_full"
+py_ver_full=${py_ver_full[1]}
+
+# Check the versions good enough
+function _check_version()
+{
+    # Bash's here document syntax is ugly
+    cat << EOF | $1
+l = lambda x: [int(y) for y in x.split('.')][0:2]
+raise SystemExit(l('${py_ver_full}') < l('${py_ver}'))
+EOF
+    return $?
+}
+
+_check_version $py_bin
+if [ "$?" -ne "0" ]; then
+    _error "Version $py_ver_full is too old. $py_ver is required"
+    echo "Examples to install Python $py_ver:"
+    echo "   Ubuntu: apt install python$py_ver"
+    echo "      Mac: brew install python@$py_ver"
+    exit 1
+fi
+echo Using version $py_ver_full
+
+# Check if our destination exists already
+dest_dir=$1/python/$py_ver_full
+if [ -d "$dest_dir" ]; then
+    exit 0
+fi
+
+# Make sure that the Python install has Pip installed
+_header "Checking for Pip"
+curl -L https://bootstrap.pypa.io/get-pip.py | $py_bin
+if [ "$?" -ne "0" ]; then
+    _error "Failed to ensure Pip is installed"
+    exit 1
+fi
+
+# Create the Python virtual env
+_header "Fetching virtual env"
+
+temp_dir=$1/python/${py_ver_full}_$$
+
+if ! $py_bin -m pip install virtualenv; then
+    _error "Unable to install 'virtualenv' Pip package"
+    _error "Please ensure Pip is installed"
+    exit 1
+fi
+
+_header "Creating a virtualenv in $temp_dir"
+if ! $py_bin -m virtualenv $temp_dir; then
+    _error "Virtualenv creation failed"
+    exit 1
+fi
+
+# Verify the virtualenv is the version we expect
+_check_version $temp_dir/bin/python
+if [ "$?" -ne "0" ]; then
+    _error "Virtualenv is an unexpected version"
+    exit 1
+fi
+
+# Swap the temp dir into place. If we fail someone else did it first
+if mv $temp_dir $dest_dir; then
+    current_dir=$(dirname $dest_dir)/current
+    _header "Symlinking $dest_dir to $current_dir"
+    rm $current_dir 2>/dev/null
+    ln -s $dest_dir $current_dir
+else
+    rm -rf $temp_dir
+fi
+
+_success Done!
+echo
+echo
diff --git a/Engine/Extras/ushell/channels/flow/posix/shells.py b/Engine/Extras/ushell/channels/flow/posix/shells.py
new file mode 100644
--- /dev/null
+++ b/Engine/Extras/ushell/channels/flow/posix/shells.py
@@ -0,0 +1,171 @@
+# Copyright Epic Games, Inc. All Rights Reserved.
+
+import os
+import re
+import shlex
+
+#-------------------------------------------------------------------------------
+class Shells(object):
+    def register_shells(self, registrar):
+        registrar.add("bash", _Bash)
+        registrar.add("zsh", _Zsh)
+        return super().register_shells(registrar)
+
+
+
+#-------------------------------------------------------------------------------
+class _Base(object):
+    def __init__(self, system):
+        self._system = system
+
+    def get_system(self):
+        return self._system
+
+    def boot_shell(self, env, cookie):
+        try: os.makedirs(os.path.dirname(cookie))
+        except: pass
+
+        with open(cookie, "wt") as out:
+            def printer(*args, **kwargs):
+                print(*args, **kwargs, file=out)
+            self.write_cookie(env, printer)
+
+    def write_cookie(self, env, out):
+        out("cd", shlex.quote(os.getcwd()))
+
+        if os.name == "nt":
+            pieces = []
+            for x in env["PATH"].split(";"):
+                x = x.replace("\\", "/")
+                if x[1] == ":":
+                    x = "/" + x[0] + x[2:]
+                pieces.append(x)
+            env["PATH"] = ":".join(pieces)
+
+        for key, value in env.read_changes():
+            if "-" not in key:
+                value = str(value).replace("\n", "\\n")
+                out(f"export {key}=$'{value}'")
+
+        out("\$tip")
+        out("echo")
+
+
+
+#-------------------------------------------------------------------------------
+class _Bash(_Base):
+    def boot_shell(self, env, cookie):
+        # Add '\[...\]' around ANSI codes so Bash knows what's displayable
+        if prompt := env.get("FLOW_PROMPT", None):
+            out_prompt = ""
+            prev = 0
+            for m in re.finditer("\x1b[^m]+m", prompt):
+                out_prompt += prompt[prev:m.start()]
+                out_prompt += rf"\[\\033{m.group()[1:]}\]"
+                prev = m.end()
+            out_prompt += prompt[prev:]
+            env["FLOW_PROMPT"] = out_prompt
+
+        return super().boot_shell(env, cookie)
+
+    def write_cookie(self, env, out):
+        super().write_cookie(env, out)
+        out("source", shlex.quote(self._add_complete_sh()))
+
+    def _add_complete_sh(self):
+        system = self.get_system()
+        working_dir = system.get_working_dir()
+        script_path = working_dir + "complete.bash"
+
+        global _complete_bash
+        cmd_tree = system.get_command_tree()
+        tree_root = cmd_tree.get_root_node()
+        for name,_ in tree_root.read_children():
+            _complete_bash += " " + shlex.quote(name)
+
+        with open(script_path, "w") as out:
+            out.write(_complete_bash)
+
+        return script_path
+
+_complete_bash = r"""
+function _flow_prompt() {
+    if [[ ! -z "$FLOW_PROMPT" ]]; then
+        local prompt=${FLOW_PROMPT/:FLOW_CWD:/$(dirs +0)}
+        PS1=$(\$complete -- \$ "$prompt")
+    fi
+}
+if [[ ! -z "$FLOW_PROMPT" ]]; then
+    PROMPT_COMMAND=_flow_prompt
+fi
+
+function _flow_complete() {
+    local choices=$(\$complete -- ${COMP_WORDS[@]}...)
+    COMPREPLY=($(compgen -W "$choices" -- ${COMP_WORDS[$COMP_CWORD]}))
+}
+complete -F _flow_complete""" # purposely left no trailing CRLF
+
+
+
+#-------------------------------------------------------------------------------
+class _Zsh(_Base):
+    def boot_shell(self, env, cookie):
+        if prompt := env.get("FLOW_PROMPT", None):
+            # Add '%{...%}' around ANSI codes so Zsh knows what's displayable
+            out_prompt = ""
+            prev = 0
+            for m in re.finditer("\x1b[^m]+m", prompt):
+                out_prompt += prompt[prev:m.start()]
+                out_prompt += "%{" + m.group() + "%}"
+                prev = m.end()
+            prompt = out_prompt + prompt[prev:]
+
+            # Convert :TAG: to $TAG
+            out_prompt = ""
+            prev = 0
+            for m in re.finditer(":[A-Z_]+:", prompt):
+                out_prompt += prompt[prev:m.start()]
+                out_prompt += "$" + m.group()[1:-1]
+                prev = m.end()
+            prompt = out_prompt + prompt[prev:]
+
+            env["FLOW_PROMPT"] = prompt
+
+        return super().boot_shell(env, cookie)
+
+    def write_cookie(self, env, out):
+        super().write_cookie(env, out)
+        out("source", shlex.quote(self._add_complete_sh()))
+
+    def _add_complete_sh(self):
+        system = self.get_system()
+        working_dir = system.get_working_dir()
+        script_path = working_dir + "complete.zsh"
+
+        cmd_tree = system.get_command_tree()
+        tree_root = cmd_tree.get_root_node()
+        flow_cmds = "' '".join(name for name,_ in tree_root.read_children())
+
+        global _complete_zsh
+        out_sh = _complete_zsh.replace("$$$FLOW_CMDS$$$", flow_cmds)
+
+        with open(script_path, "w") as out:
+            out.write(out_sh)
+
+        return script_path
+
+_complete_zsh = r"""
+_flow_prompt() {
+    FLOW_CWD=$(dirs)
+    eval $(\$prompt --format=sh)
+}
+setopt promptsubst
+if [[ $FLOW_PROMPT ]]; then
+    PROMPT=$FLOW_PROMPT
+fi
+if [[ -z "$FLOW_CMDS" ]]; then
+    precmd_functions+=(_flow_prompt)
+fi
+unset FLOW_CMDS
+FLOW_CMDS+=('$$$FLOW_CMDS$$$')
+"""
diff --git a/Engine/Extras/ushell/channels/unreal/core/boot.py b/Engine/Extras/ushell/channels/unreal/core/boot.py
new file mode 100644
--- /dev/null
+++ b/Engine/Extras/ushell/channels/unreal/core/boot.py
@@ -0,0 +1,30 @@
+# Copyright Epic Games, Inc. All Rights Reserved.
+
+import flow.cmd
+import subprocess
+
+#-------------------------------------------------------------------------------
+class Boot(flow.cmd.Cmd):
+    project = flow.cmd.Opt("", "The active project/branch. See '.project --help'")
+
+    def get_prompt(self):
+        return "\x1b[93m:UE_BRANCH:\x1b[90m|\x1b[36m:UE_PROJECT:\x1b[90m / \x1b[96m:FLOW_CWD:\n\x1b[97m-> \x1b[0m"
+
+    def run(self, env):
+        try:
+            proc = subprocess.run(
+                ("_project", self.args.project or "auto"),
+                stdin=subprocess.DEVNULL
+            )
+            if proc.returncode:
+                self.print_error()
+                self.print_error("Set a shortcut's 'Start In' to a folder with a valid .uproject")
+                self.print_error("path or launch with the 'ushell --project=[uprojpath]' argument.")
+                self.print_error()
+        except Exception as e:
+            self.print_error(e)
+        print()
+
+        env["DOTNET_CLI_TELEMETRY_OPTOUT"] = "1"
+
+        return super().run(env)
diff --git a/Engine/Extras/ushell/channels/unreal/core/cmds/build.py b/Engine/Extras/ushell/channels/unreal/core/cmds/build.py
new file mode 100644
--- /dev/null
+++ b/Engine/Extras/ushell/channels/unreal/core/cmds/build.py
@@ -0,0 +1,568 @@
+# Copyright Epic Games, Inc. All Rights Reserved.
+
+import os
+import re
+import unreal
+import tempfile
+import flow.cmd
+import unrealcmd
+import prettyprinter
+
+#-------------------------------------------------------------------------------
+def _fuzzy_find_code(prompt):
+    import fzf
+    import subprocess as sp
+
+    rg_args = (
+        "rg",
+        "--files",
+        "--path-separator=/",
+        "--no-ignore",
+        "-g*.cpp", "-g*.Build.cs",
+        "--ignore-file=" + os.path.abspath(__file__ + "/../../dirs.rgignore"),
+    )
+
+    rg = sp.Popen(rg_args, stdout=sp.PIPE, stderr=sp.DEVNULL, stdin=sp.DEVNULL)
+
+    ret = None
+    for reply in fzf.run(None, prompt=prompt, ext_stdin=rg.stdout):
+        ret = reply
+        break
+
+    rg.stdout.close()
+    rg.terminate()
+    return ret
+
+
+#-------------------------------------------------------------------------------
+class _Builder(object):
+    def __init__(self, ue_context):
+        self._args = []
+        self._ue_context = ue_context
+        self._targets = ()
+        self._projected = False
+        self.set_platform(unreal.Platform.get_host())
+        self.set_variant("development")
+
+    def set_targets(self, *targets):    self._targets = targets
+    def set_platform(self, platform):   self._platform = platform
+    def set_variant(self, variant):     self._variant = variant
+    def set_projected(self, value=True):self._projected = value
+
+    def add_args(self, *args):
+        self._args += list(args)
+
+    def read_actions(self):
+        platform = self._platform
+        variant = self._variant
+        targets = (x.get_name() for x in self._targets)
+
+        args = self._args
+        projected = any(x.is_project_target() for x in self._targets)
+        if self._projected or projected:
+            if project := self._ue_context.get_project():
+                args = ("-Project=" + str(project.get_path()), *args)
+
+        ubt = self._ue_context.get_engine().get_ubt()
+        yield from ubt.read_actions(*targets, platform, variant, *args)
+
+
+
+#-------------------------------------------------------------------------------
+class _PrettyPrinter(prettyprinter.Printer):
+    def __init__(self, logger):
+        self._source_re = re.compile(r"^(\s+\[[0-9/]+\]\s+|)([^ ]+\.[a-z]+)")
+        self._error_re = re.compile("([Ee]rror:|ERROR|Exception: |error LNK|error C)")
+        self._warning_re = re.compile("(warning C|[Ww]arning: |note:)")
+        self._progress_re = re.compile(r"^\s*@progress\s+('([^']+)'|\w+)(\s+(\d+))?")
+        self._for_target_re = re.compile(r"^\*\* For ([^\s]+) \*\*")
+        self._percent = 0
+        self._for_target = ""
+        self._section = None
+        self._logger = logger
+        self._error_file = tempfile.TemporaryFile(mode="w+t")
+
+    def print_error_summary(self):
+        import textwrap
+        import shutil
+        msg_width = shutil.get_terminal_size(fallback=(9999,0))[0]
+
+        error_file = self._error_file
+        error_file.seek(0);
+
+        files = {}
+        for line in error_file:
+            m = re.search(r"\s*(.+\.[a-z]+)[(:](\d+)[,:)0-9]+\s*(.+)$", line[1:])
+            if not m:
+                continue
+
+            line_no = int(m.group(2), 0)
+
+            lines = files.setdefault(m.group(1), {})
+            msgs = lines.setdefault(line_no, set())
+            msgs.add((line[0] != "W", m.group(3)))
+
+        if not files:
+            return
+
+        print()
+        print(flow.cmd.text.light_yellow("Errors and warnings:"))
+        print()
+        msg_width -= 6
+        for file_path, lines in files.items():
+            print(flow.cmd.text.white(file_path))
+            for line in sorted(lines.keys()):
+                msgs = lines[line]
+                print(
+                    "  Line",
+                    flow.cmd.text.white(line),
+                    flow.cmd.text.grey("/ " + os.path.basename(file_path))
+                )
+                for is_error, msg in msgs:
+                    decorator = flow.cmd.text.light_red if is_error else flow.cmd.text.light_yellow
+                    with decorator:
+                        for msg_line in textwrap.wrap(msg, width=msg_width):
+                            print("   ", msg_line)
+
+    def _print(self, line):
+        if "  Creating library " in line:
+            return
+
+        progress = self._progress_re.search(line)
+        if progress:
+            next_section = progress.group(2)
+            if next_section:
+                self._percent = min(int(progress.group(3)), 100)
+                if next_section != self._section:
+                    self._section = next_section
+                    self._logger.print_info(self._section)
+            return
+
+        if self._error_re.search(line):
+            line = line.lstrip()
+            print("E", line, sep="", file=self._error_file)
+            print(flow.cmd.text.light_red(line))
+            return
+
+        if self._warning_re.search(line):
+            line = line.lstrip()
+            if "note:" not in line:
+                print("W", line, sep="", file=self._error_file)
+            print(flow.cmd.text.light_yellow(line))
+            return
+
+        for_target = self._for_target_re.search(line)
+        if for_target:
+            for_target = for_target.group(1)
+            self._for_target = for_target[:5]
+            return
+
+        out = ""
+        if self._section != None:
+            out = flow.cmd.text.grey("%3d%%" % self._percent) + " "
+
+        source = self._source_re.search(line)
+        if source:
+            spos, epos = source.span(2)
+            if self._for_target:
+                out += flow.cmd.text.grey(self._for_target)
+                out += " "
+            out += line[:spos]
+            out += flow.cmd.text.white(line[spos:epos])
+            out += line[epos:]
+        else:
+            out += line
+        print(out)
+
+
+
+#-------------------------------------------------------------------------------
+class _BuildCmd(unrealcmd.Cmd):
+    """ The 'fileormod' argument allows the build to be constrained to a specific
+    module or source file to. It accepts one of the following forms as input;
+
+      CoreTest                 - builds the target's "CoreTest" module
+      CoreTest/SourceFile.cpp  - compiles CoreTest/Private/**/SourceFile.cpp
+      path/to/a/SourceFile.cpp - absolute or relative single file to compile
+      -                        - fuzzy find the file/modules to build (single dash)
+    """
+    _variant  = unrealcmd.Arg("development", "The build variant to have the build build")
+    fileormod = unrealcmd.Arg("", "Constrain build to a single file or module")
+    ubtargs   = unrealcmd.Arg([str], "Arguments to pass to UnrealBuildTool")
+    clean     = unrealcmd.Opt(False, "Cleans the current build")
+    nouht     = unrealcmd.Opt(False, "Skips building UnrealHeaderTool")
+    noxge     = unrealcmd.Opt(False, "Disables IncrediBuild")
+    analyze   = unrealcmd.Opt(("", "visualcpp"), "Run static analysis (add =pvsstudio for PVS)")
+    projected = unrealcmd.Opt(False, "Add the -Project= argument when running UnrealBuildTool")
+
+    complete_analyze = ("visualcpp", "pvsstudio")
+    #complete_ubtargs = ...see end of file
+
+    def complete_fileormod(self, prefix):
+        search_roots = (
+            "Source/Runtime/*",
+            "Source/Developer/*",
+            "Source/Editor/*",
+        )
+
+        ue_context = self.get_unreal_context()
+
+        if prefix:
+            from pathlib import Path
+            for root in search_roots:
+                for suffix in ("./Private/", "./Internal/"):
+                    haystack = root[:-1] + prefix + suffix
+                    yield from (x.name for x in ue_context.glob(haystack + "**/*.cpp"))
+        else:
+            seen = set()
+            for root in search_roots:
+                for item in (x for x in ue_context.glob(root) if x.is_dir()):
+                    if item.name not in seen:
+                        seen.add(item.name)
+                        yield item.name
+                        yield item.name + "/"
+
+    def _add_file_or_module(self, builder):
+        file_or_module = self.args.fileormod
+        self._constrained_build = bool(file_or_module)
+
+        if not file_or_module:
+            return
+
+        if file_or_module == "-":
+            print("Select .cpp/.Build.cs...", end="")
+            file_or_module = _fuzzy_find_code(".build [...] ")
+            if not file_or_module:
+                raise SystemExit
+
+        if file_or_module.endswith(".Build.cs"):
+            file_or_module = os.path.basename(file_or_module)[:-9]
+
+        ue_context = self.get_unreal_context()
+
+        all_modules = True
+        if project := ue_context.get_project():
+            project_path_str = str(project.get_dir())
+            all_modules = any(x in project_path_str for x in ("EngineTest", "Samples"))
+
+        if all_modules:
+            builder.add_args("-AllModules")
+
+        # If there's no ./\ in the argument then it is deemed to be a module
+        if not any((x in ".\\/") for x in file_or_module):
+            builder.add_args("-Module=" + file_or_module)
+            return
+
+        # We're expecting a file
+        if os.path.isfile(file_or_module):
+            file_or_module = os.path.abspath(file_or_module)
+            # (...N.B. an elif follows)
+
+        # Perhaps this is in module/file shorthand?
+        elif shorthand := file_or_module.split("/"):
+            rglob_pattern = "/".join(shorthand[1:])
+            search_roots = (
+                "Source/Runtime/",
+                "Source/Developer/",
+                "Source/Editor/",
+            )
+            for root in search_roots:
+                for item in ue_context.glob(root + shorthand[0]):
+                    if file := next(item.rglob(rglob_pattern), None):
+                        file_or_module = str(file.resolve())
+                        break
+
+        builder.add_args("-SingleFile=" + file_or_module)
+
+    def is_constrained_build(self):
+        return getattr(self, "_constrained_build", False)
+
+    def get_builder(self):
+        ue_context = self.get_unreal_context()
+        builder = _Builder(ue_context)
+
+        if self.args.clean: builder.add_args("-Clean")
+        if self.args.nouht: builder.add_args("-NoBuildUHT")
+        if self.args.noxge: builder.add_args("-NoXGE")
+        if self.args.analyze: builder.add_args("-StaticAnalyzer=" + self.args.analyze)
+        if self.args.ubtargs: builder.add_args(*self.args.ubtargs)
+        if self.args.projected: builder.set_projected() # forces -Project=
+        builder.add_args("-Progress")
+
+        # Handle shortcut to compile a specific module or individual file.
+        self._add_file_or_module(builder)
+
+        return builder
+
+    def run_build(self, builder):
+        try: exec_context = self._exec_context
+        except AttributeError: self._exec_context = self.get_exec_context()
+        exec_context = self._exec_context
+
+        # Inform the user of the current build configuration
+        self.print_info("Build configuration")
+        ue_context = self.get_unreal_context()
+        ubt = ue_context.get_engine().get_ubt()
+        for config in (x for x in ubt.read_configurations() if x.exists()):
+            level = config.get_level()
+            try:
+                for category, name, value in config.read_values():
+                    name = flow.cmd.text.white(name)
+                    print(category, ".", name, " = ", value, " (", level.name.title(), ")", sep="")
+            except IOError as e:
+                self.print_warning(f"Failed loading '{config.get_path()}'")
+                self.print_warning(str(e))
+                continue
+
+        self.print_info("Running UnrealBuildTool")
+        printer = _PrettyPrinter(self)
+        try:
+            for cmd, args in builder.read_actions():
+                cmd = exec_context.create_runnable(cmd, *args)
+                printer.run(cmd)
+
+                if ret := cmd.get_return_code():
+                    break
+            else:
+                return
+        except KeyboardInterrupt:
+            raise KeyboardInterrupt("Waiting for UBT...")
+        finally:
+            printer.print_error_summary()
+
+        return ret
+
+
+
+
+#-------------------------------------------------------------------------------
+class Build(_BuildCmd, unrealcmd.MultiPlatformCmd):
+    """ Direct access to running UBT """
+    target   = unrealcmd.Arg(str, "Name of the target to build")
+    platform = unrealcmd.Arg("", "Platform to build (default = host)")
+    variant  = _BuildCmd._variant
+
+    def complete_target(self, prefix):
+        ue_context = self.get_unreal_context()
+        yield from (x.name[:-10] for x in ue_context.glob("Source/*.Target.cs"))
+        yield from (x.name[:-10] for x in ue_context.glob("Source/Programs/**/*.Target.cs"))
+
+    @unrealcmd.Cmd.summarise
+    def main(self):
+        self.use_all_platforms()
+
+        ue_context = self.get_unreal_context()
+        targets = self.args.target.split()
+        targets = map(ue_context.get_target_by_name, targets)
+
+        builder = self.get_builder()
+        builder.set_targets(*targets)
+        builder.set_variant(self.args.variant)
+        if self.args.platform:
+            try: platform = self.get_platform(self.args.platform).get_name()
+            except ValueError: platform = self.args.platform
+            builder.set_platform(platform)
+        return self.run_build(builder)
+
+
+
+#-------------------------------------------------------------------------------
+class Editor(_BuildCmd, unrealcmd.MultiPlatformCmd):
+    """ Builds the editor. By default the current host platform is targeted. This
+    can be overridden using the '--platform' option to cross-compile the editor."""
+    variant     = _BuildCmd._variant
+    noscw       = unrealcmd.Opt(False, "Skip building ShaderCompileWorker")
+    nopak       = unrealcmd.Opt(False, "Do not build UnrealPak")
+    nointworker = unrealcmd.Opt(False, "Exclude the implicit InterchangeWorker dependency")
+    platform    = unrealcmd.Opt("", "Platform to build the editor for")
+
+    @unrealcmd.Cmd.summarise
+    def main(self):
+        self.use_all_platforms()
+
+        builder = self.get_builder()
+        builder.set_variant(self.args.variant)
+
+        if self.args.platform:
+            platform = self.get_platform(self.args.platform).get_name()
+            builder.set_platform(platform)
+
+        ue_context = self.get_unreal_context()
+        editor_target = ue_context.get_target_by_type(unreal.TargetType.EDITOR)
+        targets = (editor_target,)
+
+        editor_only = self.is_constrained_build() # i.e. single file or module
+        editor_only |= bool(self.args.analyze)
+        editor_only |= self.args.variant != "development"
+        if not editor_only:
+            if not self.args.noscw:
+                targets = (*targets, ue_context.get_target_by_name("ShaderCompileWorker"))
+            if not self.args.nopak:
+                targets = (*targets, ue_context.get_target_by_name("UnrealPak"))
+            if not self.args.nointworker:
+                prog_target = ue_context.get_target_by_name("InterchangeWorker")
+                if prog_target.get_type() == unreal.TargetType.PROGRAM:
+                    targets = (*targets, prog_target)
+
+        builder.set_targets(*targets)
+        return self.run_build(builder)
+
+
+
+
+#-------------------------------------------------------------------------------
+class Program(_BuildCmd, unrealcmd.MultiPlatformCmd):
+    """ Builds a program for the current host platform (unless the --platform
+    argument is provided) """
+    program   = unrealcmd.Arg(str, "Target name of the program to build")
+    variant   = _BuildCmd._variant
+    platform  = unrealcmd.Opt("", "Platform to build the program for")
+
+    def complete_program(self, prefix):
+        ue_context = self.get_unreal_context()
+
+        for depth in ("*", "*/*"):
+            for target_cs in ue_context.glob(f"Source/Programs/{depth}/*.Target.cs"):
+                yield target_cs.name[:-10]
+
+    @unrealcmd.Cmd.summarise
+    def main(self):
+        ue_context = self.get_unreal_context()
+        target = ue_context.get_target_by_name(self.args.program)
+
+        builder = self.get_builder()
+        builder.set_targets(target)
+        builder.set_variant(self.args.variant)
+
+        self.use_all_platforms()
+
+        if self.args.platform:
+            platform = self.get_platform(self.args.platform).get_name()
+            builder.set_platform(platform)
+
+        return self.run_build(builder)
+
+
+
+#-------------------------------------------------------------------------------
+class Server(_BuildCmd, unrealcmd.MultiPlatformCmd):
+    """ Builds the server target for the host platform """
+    platform = unrealcmd.Arg("", "Platform to build the server for")
+    variant  = _BuildCmd._variant
+
+    complete_platform = ("win64", "linux")
+
+    @unrealcmd.Cmd.summarise
+    def main(self):
+        # Set up the platform, defaulting to the host if none was given.
+        platform = self.args.platform
+        if not platform:
+            platform = unreal.Platform.get_host()
+        self.use_platform(self.args.platform)
+
+        ue_context = self.get_unreal_context()
+        target = ue_context.get_target_by_type(unreal.TargetType.SERVER)
+
+        builder = self.get_builder()
+        builder.set_targets(target)
+        builder.set_variant(self.args.variant)
+        if self.args.platform:
+            platform = self.get_platform(self.args.platform).get_name()
+            builder.set_platform(platform)
+
+        return self.run_build(builder)
+
+
+
+#-------------------------------------------------------------------------------
+class _Runtime(_BuildCmd, unrealcmd.MultiPlatformCmd):
+    platform = unrealcmd.Arg("", "Build the runtime for this platform")
+    variant  = _BuildCmd._variant
+
+    @unrealcmd.Cmd.summarise
+    def _main_impl(self, target):
+        platform = self.args.platform
+        if not platform:
+            platform = unreal.Platform.get_host().lower()
+        self.use_platform(platform)
+
+        platform = self.get_platform(platform).get_name()
+
+        builder = self.get_builder()
+        builder.set_platform(platform)
+        builder.set_targets(target)
+        builder.set_variant(self.args.variant)
+
+        return self.run_build(builder)
+
+#-------------------------------------------------------------------------------
+class Client(_Runtime):
+    """ Builds the client runtime """
+    def main(self):
+        ue_context = self.get_unreal_context()
+        target = ue_context.get_target_by_type(unreal.TargetType.CLIENT)
+        return super()._main_impl(target)
+
+#-------------------------------------------------------------------------------
+class Game(_Runtime):
+    """ Builds the game runtime target binary executable """
+    def main(self):
+        ue_context = self.get_unreal_context()
+        target = ue_context.get_target_by_type(unreal.TargetType.GAME)
+        return super()._main_impl(target)
+
+
+
+#-------------------------------------------------------------------------------
+class _CleanImpl(object):
+    def main(self):
+        self.args.clean = True
+        return super().main()
+
+def _make_clean_command(base):
+    return type("CleanImpl", (_CleanImpl, base), {})
+
+CleanEditor  = _make_clean_command(Editor)
+CleanProgram = _make_clean_command(Program)
+CleanServer  = _make_clean_command(Server)
+CleanClient  = _make_clean_command(Client)
+CleanGame    = _make_clean_command(Game)
+
+
+#-------------------------------------------------------------------------------
+_BuildCmd.complete_ubtargs = (
+    "-2015", "-2017", "-2019", "-AllModules", "-alwaysgeneratedsym",
+    "-Architectures", "-build-as-framework", "-BuildPlugin", "-BuildVersion",
+    "-BundleVersion", "-Clean", "-CLion", "-CMakefile", "-CodeliteFiles",
+    "-CompileAsDll", "-CompileChaos", "-CompilerArguments", "-CompilerVersion",
+    "-CppStd", "-CreateStub", "-Define:", "-DependencyList", "-Deploy",
+    "-DisableAdaptiveUnity", "-DisablePlugin", "-DisableUnity", "-distribution",
+    "-dSYM", "-EdditProjectFiles", "-EnableASan", "-EnableMSan",
+    "-EnablePlugin", "-EnableTSan", "-EnableUBSan",
+    "-FailIfGeneratedCodeChanges", "-FastMonoCalls", "-FastPDB", "-FlushMac",
+    "-ForceDebugInfo", "-ForceHeaderGeneration", "-ForceHotReload",
+    "-ForceUnity", "-Formal", "-FromMsBuild", "-generatedsymbundle",
+    "-generatedsymfile", "-GPUArchitectures", "-IgnoreJunk",
+    "-ImportCertificate", "-ImportCertificatePassword", "-ImportProvision",
+    "-IncrementalLinking", "-Input", "-IWYU", "-KDevelopfile",
+    "-LinkerArguments", "-LiveCoding", "-LiveCodingLimit",
+    "-LiveCodingManifest", "-LiveCodingModules", "-Log", "-LTCG", "-Makefile",
+    "-Manifest", "-MapFile", "-MaxParallelActions", "-Modular", "-Module",
+    "-Monolithic", "-NoBuildUHT", "-NoCompileChaos", "-NoDebugInfo", "-NoDSYM",
+    "-NoEngineChanges", "-NoFASTBuild", "-NoFastMonoCalls", "-NoHotReload",
+    "-NoHotReloadFromIDE", "-NoIncrementalLinking", "-NoLink",
+    "-NoManifestChanges", "-NoMutex", "-NoPCH", "-NoPDB", "-NoSharedPCH",
+    "-NoUBTMakefiles", "-NoUseChaos", "-NoXGE", "-ObjSrcMap", "-Output",
+    "-OverrideBuildEnvironment", "-PGOOptimize", "-PGOProfile", "-Plugin",
+    "-Precompile", "-Preprocess", "-PrintDebugInfo", "-Progress",
+    "-ProjectDefine:", "-ProjectFileFormat", "-ProjectFiles",
+    "-PublicSymbolsByDefault", "-QMakefile", "-Quiet", "-RemoteIni", "-Rider",
+    "-rtti", "-ShadowVariableErrors", "-SharedBuildEnvironment",
+    "-ShowIncludes", "-SingleFile", "-SkipBuild", "-skipcrashlytics",
+    "-SkipDeploy", "-SkipPreBuildTargets", "-SkipRulesCompile",
+    "-StaticAnalyzer", "-StressTestUnity", "-Strict", "-stripsymbols",
+    "-ThinLTO", "-Timestamps", "-Timing", "-ToolChain", "-Tracing",
+    "-UniqueBuildEnvironment", "-UseChaos", "-UsePrecompiled", "-Verbose",
+    "-VeryVerbose", "-VSCode", "-VSMac", "-WaitMutex", "-WarningsAsErrors",
+    "-WriteOutdatedActions", "-XCodeProjectFiles", "-XGEExport",
+)
diff --git a/Engine/Extras/ushell/channels/unreal/core/cmds/build_xml.py b/Engine/Extras/ushell/channels/unreal/core/cmds/build_xml.py
new file mode 100644
--- /dev/null
+++ b/Engine/Extras/ushell/channels/unreal/core/cmds/build_xml.py
@@ -0,0 +1,164 @@
+# Copyright Epic Games, Inc. All Rights Reserved.
+
+import os
+import unreal
+import flow.cmd
+import unrealcmd
+
+#-------------------------------------------------------------------------------
+class Show(unrealcmd.Cmd):
+    """ Display the current build configuration for the current project. To
+    modify the configuration use the sub-commands set, clear, and edit. """
+
+    def main(self):
+        root = {}
+
+        context = self.get_unreal_context()
+        ubt = context.get_engine().get_ubt()
+        for config in (x for x in ubt.read_configurations() if x.exists()):
+            level = config.get_level()
+            try:
+                for section, name, value in config.read_values():
+                    section = root.setdefault(section, {})
+                    name = section.setdefault(name, [])
+                    name.append((value, level))
+            except IOError as e:
+                self.print_warning(f"Failed loading '{config.get_path()}'")
+                self.print_warning(str(e))
+                continue
+
+        col0 = max(len(x) for x in root.keys()) if root else 0
+        col0 += 2
+
+        col1 = 0
+        for section in root.values():
+            n = max(len(x) for x in section.keys())
+            col1 = max(col1, n)
+        col1 += 2
+
+        def spam(n, x, decorator=None):
+            dots = "." * (n - len(x))
+            if decorator: x = decorator(x)
+            print(x, flow.cmd.text.grey(dots), "", end="")
+
+        self.print_info("Current configuration")
+        if not root:
+            print("There isn't any!")
+        for section, payload in root.items():
+            for name, payload in payload.items():
+                for value, level in payload:
+                    value = value + " (" + level.name.title() + ")"
+                    spam(col0, section, flow.cmd.text.cyan)
+                    spam(col1, name, flow.cmd.text.white)
+                    spam(0, value)
+                    print()
+
+        self.print_info("File paths")
+        for config in ubt.read_configurations():
+            red_green = flow.cmd.text.green if config.exists() else flow.cmd.text.red
+            spam(10, config.get_level().name.title(), decorator=red_green)
+            spam(0, os.path.normpath(config.get_path()))
+            print()
+
+#-------------------------------------------------------------------------------
+class _SetClearBase(unrealcmd.Cmd):
+    section  = unrealcmd.Arg(str, "Configuration section to modify")
+    name     = unrealcmd.Arg(str, "The name of the value to change")
+    branch   = unrealcmd.Opt(False, "Only apply changes to this branch (otherwise global)")
+
+    def _complete_schema(self, prefix):
+        context = self.get_unreal_context()
+        ubt = context.get_engine().get_ubt()
+        schema = ubt.get_configuration_schema()
+
+        def read_node_names():
+            if self.args.section: yield self.args.section
+            if self.args.name: yield self.args.name
+
+        if node := schema.get_node(*read_node_names()):
+            yield from (x.get_name() for x in node.read_children())
+
+    complete_section  = _complete_schema
+    complete_name = _complete_schema
+
+    def _load_configuration(self):
+        self.print_info("Modifying build configuration")
+
+        context = self.get_unreal_context()
+        ubt = context.get_engine().get_ubt()
+        configuration = ubt.get_configuration(self.args.branch, create=True)
+
+        print("Scope:", configuration.get_level().name.title())
+        print("Path:", configuration.get_path())
+
+        return configuration
+
+
+#-------------------------------------------------------------------------------
+class _SetImpl(object):
+    """ Modify build configuration. """
+    value = unrealcmd.Arg(str, "A value to set the property to")
+
+    complete_value = _SetClearBase._complete_schema
+
+    def main(self):
+        section, name, value = self.args.section, self.args.name, self.args.value
+        configuration = self._load_configuration()
+
+        prev_value = None
+        try:
+            prev_value = configuration.get_value(section, name)
+        except LookupError:
+            pass
+
+        configuration.set_value(section, name, value)
+        configuration.save()
+
+        print(f"Setting '{section}.{name}' to '{value}' (was '{prev_value}')")
+        with flow.cmd.text.green:
+            print("Success")
+
+#-------------------------------------------------------------------------------
+# This exists simply to get arguments (i.e. unreal.Arg) in the correct order.
+class Set(_SetClearBase, _SetImpl):
+    pass
+
+#-------------------------------------------------------------------------------
+class Clear(_SetClearBase):
+    """ Clears a value from the build configuration. """
+
+    def main(self):
+        section, name = self.args.section, self.args.name
+        configuration = self._load_configuration()
+
+        prev_value = None
+        try:
+            prev_value = configuration.get_value(section, name)
+        except LookupError:
+            pass
+
+        try:
+            configuration.clear_value(section, name)
+            configuration.save()
+        except LookupError as e:
+            self.print_error(e)
+            return False
+
+        print(f"Cleared '{section}.{name}' (was '{prev_value}')")
+        with flow.cmd.text.green:
+            print("Success")
+
+#-------------------------------------------------------------------------------
+class Edit(unrealcmd.Cmd):
+    """ Opens a build configuration in a text editor. The text editor used is
+    determined in the following order; $GIT_EDITOR, $P4EDITOR, system default. """
+    branch = unrealcmd.Opt(False, "Edit the build configuration for this branch")
+
+    def main(self):
+        context = self.get_unreal_context()
+        ubt = context.get_engine().get_ubt()
+        configuration = ubt.get_configuration(self.args.branch, create=True)
+
+        path = configuration.get_path()
+        print("Editing", path)
+        self.edit_file(path)
diff --git a/Engine/Extras/ushell/channels/unreal/core/cmds/clangdb.py b/Engine/Extras/ushell/channels/unreal/core/cmds/clangdb.py
new file mode 100644
--- /dev/null
+++ b/Engine/Extras/ushell/channels/unreal/core/cmds/clangdb.py
@@ -0,0 +1,106 @@
+# Copyright Epic Games, Inc. All Rights Reserved.
+
+import os
+import unreal
+import unrealcmd
+
+#-------------------------------------------------------------------------------
+def _find_autosdk_llvm(autosdk_path):
+    path = f"{autosdk_path}/Host{unreal.Platform.get_host()}/Win64/LLVM"
+
+    best_version = None
+    if os.path.isdir(path):
+        with os.scandir(path) as dirs:
+            best_version = max((x.name for x in dirs if x.is_dir()), default=None)
+
+    path += f"/{best_version}/"
+    if not os.path.isfile(path + "bin/clang-cl.exe"):
+        best_version = path = None
+
+    return (best_version, path)
+
+#-------------------------------------------------------------------------------
+def _can_execute_clang():
+    import subprocess
+    try: subprocess.run(("clang-cl", "--version"), stdout=subprocess.DEVNULL)
+    except: return False
+    return True
+
+
+
+#-------------------------------------------------------------------------------
+class ClangDb(unrealcmd.Cmd):
+    """ Generates a Clang database (compile_commands.json). If there is an active
+    project then the database is generated in context of that project. By default
+    the editor target is used. """
+    target    = unrealcmd.Arg("", "Target to use when generating the database")
+    ubtargs   = unrealcmd.Arg([str], "Arguments to pass onwards to UnrealBuildTool")
+    platform  = unrealcmd.Opt("", "Platform to build a database for")
+
+    @unrealcmd.Cmd.summarise
+    def _run_ubt(self, target, platform, ubt_args):
+        ue_context = self.get_unreal_context()
+        ubt = ue_context.get_engine().get_ubt()
+        if project := ue_context.get_project():
+            ubt_args = ("-Project=" + str(project.get_path()), *ubt_args)
+
+        exec_context = self.get_exec_context()
+        for cmd, args in ubt.read_actions(target, platform, "Development", *ubt_args):
+            cmd = exec_context.create_runnable(cmd, *args)
+            cmd.run()
+            ret = cmd.get_return_code()
+            if ret:
+                return ret
+
+    def main(self):
+        # Rough check for the presence of LLVM
+        self.print_info("Checking for LLVM")
+        autosdk_path = os.getenv("UE_SDKS_ROOT")
+        print("UE_SDKS_ROOT:", autosdk_path)
+        llvm_version, llvm_path = _find_autosdk_llvm(autosdk_path)
+        if llvm_version:
+            print("Found:", llvm_version, "at", llvm_path)
+        else:
+            print("No suitable AusoSDKs version found")
+            print("Trying PATH ... ", end="")
+            has_clang = _can_execute_clang()
+            print("ok" if has_clang else "nope")
+            if not has_clang:
+                self.print_warning("Unable to locate Clang binary")
+
+        ue_context = self.get_unreal_context()
+
+        # Calculate how to invoke UBT
+        target = self.args.target
+        if not target:
+            target = ue_context.get_target_by_type(unreal.TargetType.EDITOR)
+            target = target.get_name()
+
+        platform = unreal.Platform.get_host()
+        if self.args.platform:
+            platform = self.get_platform(self.args.platform).get_name()
+
+        self.print_info("Database properties")
+        print("Target:", target)
+        print("Platform:", platform)
+
+        ubt_args = (
+            "-allmodules",
+            "-Mode=GenerateClangDatabase",
+            *self.args.ubtargs,
+        )
+
+        # Run UBT
+        self.print_info("Running UnrealBuildTool")
+        ret = self._run_ubt(target, platform, ubt_args)
+        if ret:
+            return ret
+
+        # Print some information
+        db_path = ue_context.get_engine().get_dir() / "../compile_commands.json"
+        try: db_size = os.path.getsize(db_path)
+        except: db_size = -1
+        print()
+        self.print_info("Database details")
+        print("Path:", os.path.normpath(db_path))
+        print("Size: %.02fKB" % (db_size / 1024 / 1024))
diff --git a/Engine/Extras/ushell/channels/unreal/core/cmds/cook.py b/Engine/Extras/ushell/channels/unreal/core/cmds/cook.py
new file mode 100644
--- /dev/null
+++ b/Engine/Extras/ushell/channels/unreal/core/cmds/cook.py
@@ -0,0 +1,165 @@
+# Copyright Epic Games, Inc. All Rights Reserved.
+
+import re
+import unreal
+import flow.cmd
+import unrealcmd
+import uelogprinter
+import unreal.cmdline
+
+#-------------------------------------------------------------------------------
+class _CookPrettyPrinter(uelogprinter.Printer):
+    def __init__(self):
+        super().__init__()
+        self._cook_re = re.compile(r"ages (\d+).+otal (\d+)")
+
+    def _print(self, line):
+        m = self._cook_re.search(line)
+        if m:
+            percent = (int(m.group(1)) * 100) // max(int(m.group(2)), 1)
+            line += f" [{percent}%]"
+        super()._print(line)
+
+
+
+#-------------------------------------------------------------------------------
+class _Cook(unrealcmd.MultiPlatformCmd):
+    cookargs = unrealcmd.Arg([str], "Additional arguments passed to the cook commandlet")
+    cultures = unrealcmd.Opt("en", "Cultures to cook (comma separated, defaults to 'en')")
+    onthefly = unrealcmd.Opt(False, "Launch as an on-the-fly server")
+    iterate  = unrealcmd.Opt(False, "Cook iteratively on top of the previous cook")
+    noxge    = unrealcmd.Opt(False, "Disable XGE-based shader compilation")
+    unpretty = unrealcmd.Opt(False, "Turns off colourful pretty-printing")
+    attach   = unrealcmd.Opt(False, "Attach a debugger to the cook")
+    debug    = unrealcmd.Opt(False, "Use debug executables")
+
+    @unrealcmd.Cmd.summarise
+    def _cook(self, cook_form, exec_context=None):
+        ue_context = self.get_unreal_context()
+
+        # Check there is a valid project
+        project = ue_context.get_project()
+        if not project:
+            self.print_error("An active project is required to cook")
+            return False
+
+        # Prepare arguments
+        cultures = self.args.cultures.replace(",", "+")
+
+        commandlet_args = (
+            "-targetplatform=" + cook_form,
+            "-cookcultures=" + cultures if cultures else None,
+            "-unattended",
+            "-unversioned",
+            "-stdout",
+            "-cookonthefly" if self.args.onthefly else None,
+            "-iterate" if self.args.iterate else None,
+            "-noxgeshadercompile" if self.args.noxge else None,
+            *unreal.cmdline.read_ueified(*self.args.cookargs),
+        )
+
+        commandlet_args = tuple(x for x in commandlet_args if x)
+
+        # To support attaching a debugger we will reuse '.run commandlet'. This
+        # approach is not used normally as it interferes with pretty-printing and
+        # creates an awful lot of child processes.
+
+        if self.args.attach:
+            args = ("_run", "commandlet", "cook", "--noplatforms", "--attach")
+            if self.args.debug:
+                args = (*args, "debug")
+            args = (*args, "--", *commandlet_args)
+
+            # Disable our handling of Ctrl-C
+            import signal
+            def nop_handler(*args): pass
+            signal.signal(signal.SIGINT, nop_handler)
+
+            import subprocess
+            return subprocess.run(args).returncode
+
+        # Find the editor binary
+        target = ue_context.get_target_by_type(unreal.TargetType.EDITOR)
+        variant = unreal.Variant.DEBUG if self.args.debug else unreal.Variant.DEVELOPMENT
+        build = target.get_build(variant=variant)
+        if not build:
+            self.print_error(f"No {variant.name.lower()} editor build found")
+            return False
+
+        cmd = str(build.get_binary_path())
+        if not cmd.endswith("-Cmd.exe"):
+            cmd = cmd.replace(".exe", "-Cmd.exe")
+
+        # Launch
+        args = (
+            project.get_path(),
+            "-run=cook",
+            *commandlet_args,
+        )
+
+        exec_context = exec_context or self.get_exec_context()
+        cmd = exec_context.create_runnable(cmd, *args)
+        if self.args.unpretty or not self.is_interactive():
+            cmd.run()
+        else:
+            printer = _CookPrettyPrinter()
+            printer.run(cmd)
+        return cmd.get_return_code()
+
+
+
+#-------------------------------------------------------------------------------
+class Cook(_Cook):
+    """ Prepares a cooked meal of dataset for the given target platform """
+    target = unrealcmd.Arg(str, "Full target platform name (i.e. WindowsNoEditor)")
+    def main(self):
+        return self._cook(self.args.target)
+
+
+
+#-------------------------------------------------------------------------------
+class _Runtime(_Cook):
+    platform = unrealcmd.Arg("", "The platform to cook data for")
+
+    def main(self, target):
+        platform = self.args.platform
+        if not platform:
+            platform = unreal.Platform.get_host()
+        platform = self.get_platform(platform)
+
+        exec_context = super().get_exec_context()
+        exec_env = exec_context.get_env()
+
+        # Establish target platform's environment
+        name = platform.get_name()
+        self.print_info(name, "-", platform.get_version())
+        for item in platform.read_env():
+            try:
+                dir = str(item)
+                exec_env[item.key] = dir
+            except EnvironmentError:
+                dir = flow.cmd.text.red(item.get())
+            print(item.key, "=", dir)
+
+        cook_form = platform.get_cook_form(target.name.lower())
+        return self._cook(cook_form, exec_context)
+
+#-------------------------------------------------------------------------------
+class Game(_Runtime):
+    """ Cooks game data for a particular platform """
+    def main(self):
+        return super().main(unreal.TargetType.GAME)
+
+#-------------------------------------------------------------------------------
+class Client(_Runtime):
+    """ Cooks client data for a particular platform """
+    def main(self):
+        return super().main(unreal.TargetType.CLIENT)
+
+#-------------------------------------------------------------------------------
+class Server(_Runtime):
+    """ Cooks server data for a particular platform """
+    complete_platform = ("win64", "linux")
+
+    def main(self):
+        return super().main(unreal.TargetType.SERVER)
diff --git a/Engine/Extras/ushell/channels/unreal/core/cmds/ddc.py b/Engine/Extras/ushell/channels/unreal/core/cmds/ddc.py
new file mode 100644
--- /dev/null
+++ b/Engine/Extras/ushell/channels/unreal/core/cmds/ddc.py
@@ -0,0 +1,59 @@
+# Copyright Epic Games, Inc. All Rights Reserved.
+
+import sys
+import unrealcmd
+import subprocess
+
+#-------------------------------------------------------------------------------
+class Auth(unrealcmd.Cmd):
+    """ Authorize oneself for cloud-based DDC """
+    service     = unrealcmd.Arg("", "Name of the service to authorize with")
+    extraargs   = unrealcmd.Arg([str], "Authorisation tool arguments")
+    query       = unrealcmd.Opt(False, "Only check the current status")
+
+    def _lookup_service_name(self):
+        if getattr(self, "_service_name", None):
+            return
+
+        self.print_info("Discovering service name")
+        print("Config path: Engine/StorageServers/Default")
+
+        config = self.get_unreal_context().get_config()
+        default = config.get("Engine", "StorageServers", "Default")
+        if value := default.OAuthProviderIdentifier:
+            print("found;", value)
+            self._service_name = value
+
+    def main(self):
+        ue_context = self.get_unreal_context()
+        engine = ue_context.get_engine()
+
+        bin_type = "win-x64"
+        if sys.platform == "darwin": bin_type = "osx-x64"
+        elif sys.platform == "linux": bin_type = "linux-x64"
+
+        bin_path = engine.get_dir()
+        bin_path /= f"Binaries/DotNET/OidcToken/{bin_type}/OidcToken.exe"
+        if not bin_path.is_file():
+            raise FileNotFoundError(f"Unable to find '{bin_path}'")
+
+        self._service_name = self.args.service
+        self._lookup_service_name()
+        if not self._service_name:
+            raise ValueError("Unable to discover service name")
+
+        args = (
+            "--Service=" + str(self._service_name),
+        )
+
+        if self.args.query:
+            args = (*args, "--Mode=Query")
+
+        if project := ue_context.get_project():
+            args = (*args, "--Project=" + str(project.get_path()))
+
+        args = (*args, *self.args.extraargs)
+
+        cmd = self.get_exec_context().create_runnable(bin_path, *args)
+        result = cmd.run()
+        return result
diff --git a/Engine/Extras/ushell/channels/unreal/core/cmds/info.py b/Engine/Extras/ushell/channels/unreal/core/cmds/info.py
new file mode 100644
--- /dev/null
+++ b/Engine/Extras/ushell/channels/unreal/core/cmds/info.py
@@ -0,0 +1,181 @@
+# Copyright Epic Games, Inc. All Rights Reserved.
+
+import unreal
+import flow.cmd
+import unrealcmd
+
+#-------------------------------------------------------------------------------
+def _write(key, value):
+    value = str(value)
+
+    if   value.startswith("?"): value = flow.cmd.text.light_yellow(value[1:])
+    elif value.startswith("!"): value = flow.cmd.text.light_red(value[1:])
+
+    dots = max(23 - len(str(key)), 3)
+    dots = flow.cmd.text.grey(("." * dots) if dots > 0 else "")
+    print(key, dots, value)
+
+
+
+#-------------------------------------------------------------------------------
+class InfoBase(unrealcmd.Cmd):
+    json = unrealcmd.Opt(False, "Print the output as JSON")
+
+    def _print_pretty(self, info):
+        def print_dict(name, subject):
+            if isinstance(subject, dict):
+                it = subject.items()
+            elif isinstance(subject, list):
+                it = enumerate(subject)
+
+            for key, value in it:
+                if isinstance(value, dict|list):
+                    print_dict(key, value)
+                    continue
+
+                if name and name != "env":
+                    print(end=self._end)
+                    self._end = "\n"
+                    self.print_info(name.title().replace("_", " "))
+                    name = None
+
+                _write(key, value)
+
+        print_dict("", info)
+
+    def _print_json(self, info):
+        import sys
+        import json
+        def represent(subject):
+            if hasattr(subject, "get_inner"):
+                return subject.get_inner()
+            return repr(subject)
+        json.dump(info, sys.stdout, indent=2, default=represent)
+
+    def main(self):
+        self._end = ""
+        info = self._get_info()
+        if self.args.json:
+            self._print_json(info)
+        else:
+            self._print_pretty(info)
+
+
+
+#-------------------------------------------------------------------------------
+class Info(InfoBase):
+    """ Displays info about the current session. """
+
+    def _get_engine_info(self):
+        context = self.get_unreal_context()
+        engine = context.get_engine()
+
+        info = {
+            "path"          : str(engine.get_dir()),
+            "version"       : "",
+            "version_full"  : "",
+            "branch"        : "",
+            "changelist"    : "0",
+        }
+
+        version = engine.get_info()
+        if version:
+            info["branch"] = version["BranchName"]
+            info["changelist"] = str(version["Changelist"])
+            info["version"] = str(engine.get_version_major())
+            info["version_full"] = "{}.{}.{}".format(
+                version["MajorVersion"],
+                version["MinorVersion"],
+                version["PatchVersion"])
+
+        return info
+
+    def _get_project_info(self):
+        context = self.get_unreal_context()
+        project = context.get_project()
+
+        info = {
+            "name" : "?no active project"
+        }
+
+        if not project:
+            return info
+
+        info["name"] = project.get_name()
+        info["path"] = str(project.get_path())
+        info["dir"] = str(project.get_dir())
+
+        targets_info = {}
+        known_targets = (
+            ("editor", unreal.TargetType.EDITOR),
+            ("game",   unreal.TargetType.GAME),
+            ("client", unreal.TargetType.CLIENT),
+            ("server", unreal.TargetType.SERVER),
+        )
+        for target_name, target_type in known_targets:
+            target = context.get_target_by_type(target_type)
+            output = target.get_name() if target else "!unknown"
+            targets_info[target_name] = output
+        info["targets"] = targets_info
+
+        return info
+
+    def _get_platforms_info(self):
+        info = {}
+
+        for platform_name in sorted(self.read_platforms()):
+            inner = info.setdefault(platform_name.lower(), {})
+            platform = self.get_platform(platform_name)
+
+            inner["name"] = platform.get_name()
+            inner["version"] = platform.get_version() or "!unknown"
+
+            inner["env"] = {}
+            for item in platform.read_env():
+                try: value = item.validate()
+                except EnvironmentError: value = "!" + item.get()
+                inner["env"][item.key] = value
+
+        return info
+
+    def _get_info(self):
+        return {
+            "engine"    : self._get_engine_info(),
+            "project"   : self._get_project_info(),
+            "platforms" : self._get_platforms_info(),
+        }
+
+
+
+#-------------------------------------------------------------------------------
+class Projects(InfoBase):
+    """ Displays the branch's synced .uprojects """
+
+    def _get_info(self):
+        context = self.get_unreal_context()
+        branch = context.get_branch(must_exist=True)
+        return { "branch_projects" : list(str(x) for x in branch.read_projects()) }
+
+
+
+#-------------------------------------------------------------------------------
+class Config(InfoBase):
+    """ Evaluates currect context's config for a particular category. The 'override'
+    option can be used to enumerate platform-specific config sets."""
+    category = unrealcmd.Arg(str, "The ini category to load")
+    override = unrealcmd.Opt("", "Alternative config set to enumerate (e.g. 'WinGDK')")
+
+    def complete_category(self, prefix):
+	    return ("Engine", "Game", "Input", "DeviceProfiles", "GameUserSettings",
+            "Scalability", "RuntimeOptions", "InstallBundle", "Hardware", "GameplayTags")
+
+    def _get_info(self):
+        context = self.get_unreal_context()
+        config = context.get_config()
+        ini = config.get(self.args.category, override=self.args.override)
+
+        ret = {}
+        for s_name, section in ini:
+            ret[s_name] = {k:v for k,v in section}
+
+        return ret
diff --git a/Engine/Extras/ushell/channels/unreal/core/cmds/kill.py b/Engine/Extras/ushell/channels/unreal/core/cmds/kill.py
new file mode 100644
--- /dev/null
+++ b/Engine/Extras/ushell/channels/unreal/core/cmds/kill.py
@@ -0,0 +1,29 @@
+# Copyright Epic Games, Inc. All Rights Reserved.
+
+import unrealcmd
+
+#-------------------------------------------------------------------------------
+class Kill(unrealcmd.Cmd):
+    """ Hard-terminates a running process on dev kits or the local machine. """
+    what = unrealcmd.Arg(str, "Specified what platform or item to kill")
+
+    def complete_what(self, prefix):
+        yield "editor"
+        yield "server"
+        yield "client"
+        yield from self.complete_platform(prefix)
+
+    def main(self):
+        what = self.args.what
+        try:
+            platform = self.get_platform(what)
+            platform.kill()
+            return
+        except ValueError:
+            pass
+
+        platform = self.get_platform()
+        if what == "editor":   platform.kill(what)
+        elif what == "server": platform.kill(what)
+        elif what == "client": platform.kill(what)
+        else: raise ValueError(f"I do not know how to '.kill {what}'")
diff --git a/Engine/Extras/ushell/channels/unreal/core/cmds/notify.py b/Engine/Extras/ushell/channels/unreal/core/cmds/notify.py
new file mode 100644
--- /dev/null
+++ b/Engine/Extras/ushell/channels/unreal/core/cmds/notify.py
@@ -0,0 +1,14 @@
+# Copyright Epic Games, Inc. All Rights Reserved.
+
+import os
+import unrealcmd
+import ctypes
+
+class Notify(unrealcmd.Cmd):
+    """ Flash the console window (if possible) to get the users attention """
+
+    def main(self):
+        if os.name == "nt":
+            ctypes.windll.user32.FlashWindow(ctypes.windll.kernel32.GetConsoleWindow(), True )
+        else:
+            print("Notify is not supported by this OS")
\ No newline at end of file
diff --git a/Engine/Extras/ushell/channels/unreal/core/cmds/project.py b/Engine/Extras/ushell/channels/unreal/core/cmds/project.py
new file mode 100644
--- /dev/null
+++ b/Engine/Extras/ushell/channels/unreal/core/cmds/project.py
@@ -0,0 +1,324 @@
+# Copyright Epic Games, Inc. All Rights Reserved.
+
+import os
+import unreal
+import flow.cmd
+
+#-------------------------------------------------------------------------------
+class Change(flow.cmd.Cmd):
+    """ Changes this session's active project. The 'nameorpath' argument can be
+    one of following;
+
+     .project MyProjectName       - use branch's .uprojectdirs and find by name
+     .project d:/foo/bar.uproject - set by .uproject path
+     .project branch              - unset the active project on the branch
+     .project cwd                 - active project or branch follows CWD
+     .project active              - output current active project
+     .project                     - pick from a list of all projects under the CWD
+    """
+    nameorpath   = flow.cmd.Arg("", "Name or path of the project to make active")
+
+    def _complete_path(self, prefix):
+        prefix = prefix or "."
+        for item in os.scandir(prefix):
+            if item.name.endswith(".uproject"):
+                yield item.name
+                return
+
+        for item in os.scandir(prefix):
+            if item.is_dir():
+                yield item.name + "/"
+
+    def complete_nameorpath(self, prefix):
+        if "/" in prefix or "\\" in prefix:
+            yield from self._complete_path(prefix)
+            return
+
+        yield "branch"
+        yield "cwd"
+        yield "active"
+
+        session = self.get_noticeboard(self.Noticeboard.SESSION)
+        context = unreal.Context(session["uproject"] or ".")
+        branch = context.get_branch()
+        for uproj_path in branch.read_projects():
+                yield uproj_path.stem
+
+    def _select_project(self):
+        if not self.is_interactive():
+            return
+
+        import fzf
+        import pathlib
+        import subprocess as sp
+
+        # Select the folder to use. If there's a .uprojectdirs up the tree we'll
+        # start there otherwise we'll use the cwd
+        try:
+            context = unreal.Context(".")
+            branch = context.get_branch()
+            cwd = str(branch.get_dir())
+            def read_projects():
+                for uproj_path in branch.read_projects():
+                    yield f"{uproj_path.stem:19}\t{uproj_path}"
+        except:
+            cwd = os.getcwd()
+            def read_projects():
+                rg_args = (
+                    "rg",
+                    "--files",
+                    "--path-separator=/",
+                    "--no-ignore",
+                    "--ignore-file=" + os.path.abspath(__file__ + "/../../uproject.rgignore"),
+                    "--ignore-file=" + os.path.abspath(__file__ + "/../../dirs.rgignore"),
+                )
+
+                rg = sp.Popen(rg_args, stdout=sp.PIPE, stderr=sp.DEVNULL)
+
+                for line in rg.stdout.readlines():
+                    line = line.decode().rstrip().split("/")
+                    if len(line) > 1:
+                        name = line[-2].replace(" ", "")
+                        yield f"{name:19}\t{'/'.join(line)}"
+                rg.wait()
+
+        print("Select a project (fuzzy, move with arrows, enter to select)...", end="")
+        prompt = os.path.normpath(cwd).replace("\\", "/") + "/"
+        for reply in fzf.run(read_projects(), prompt=prompt):
+            reply = reply.split("\t", 1)
+            reply = reply[1].lstrip()
+            return os.path.abspath(reply)
+
+    def _print_active(self):
+        session = self.get_noticeboard(self.Noticeboard.SESSION)
+        context = unreal.Context(session["uproject"] or ".")
+        branch = context.get_branch()
+        if not branch:
+            raise EnvironmentError(f"Unable to find a branch with a .uprojectdirs from '{context.get_primary_path()}'")
+
+        if project := context.get_project():
+            print(project.get_dir())
+        else:
+            self.print_info("[no active project]")
+        return
+
+    def _set_cwd(self):
+        self.print_info("Clearing active branch/project")
+        print("Context will now follow current directory")
+        session = self.get_noticeboard(self.Noticeboard.SESSION)
+        session["uproject"] = "."
+
+    def _set_branch(self):
+        self.print_info("Clearing active project")
+        session = self.get_noticeboard(self.Noticeboard.SESSION)
+        context = unreal.Context(session["uproject"] or ".")
+        branch = context.get_branch()
+        if not branch:
+            context = unreal.Context(".")
+            branch = context.get_branch()
+            if not branch:
+                raise EnvironmentError(f"Unable to clear the active project without a branch")
+        candidate = branch.get_dir()
+        print("Using branch;", candidate)
+        return candidate
+
+    def _set_auto(self):
+        try:
+            context = unreal.Context(".")
+        except EnvironmentError:
+            return
+
+        if project := context.get_project():
+            return str(project.get_dir())
+
+        if (branch := context.get_branch()) is None:
+            return
+
+        best_mtime = 0
+        best = None
+        for uproj_path in branch.read_projects():
+            mtime = uproj_path.stat().st_mtime
+            if mtime < best_mtime:
+                continue
+            best_mtime = mtime
+            best = uproj_path
+
+        if best:
+            return str(best)
+
+        return str(branch.get_dir())
+
+
+    def _set_qualified(self, candidate):
+        candidate = os.path.abspath(candidate)
+        self.print_info("Using qualified directory or file path")
+        print("Path:", candidate)
+
+        try:
+            context = unreal.Context(candidate)
+        except EnvironmentError as e:
+            # If the user gave a directory we can fallback to a fuzzy search
+            if os.path.isdir(candidate):
+                os.chdir(candidate)
+                self.args.nameorpath = ""
+                return self._main_impl()
+
+            self.print_error(e)
+            return
+
+        if project := context.get_project():
+            candidate = project.get_path()
+            print("Project found;", candidate)
+        else:
+            candidate = context.get_branch().get_dir()
+            print("Branch found;", candidate)
+
+        return candidate
+
+    def _set_short_form(self, candidate):
+        self.print_info("Short-form search")
+
+        # Get an Unreal context, preferring the CWD and falling back to the
+        # current active context if there is one.
+        try:
+            context = unreal.Context(".")
+        except EnvironmentError:
+            session = self.get_noticeboard(self.Noticeboard.SESSION)
+            context = unreal.Context(session["uproject"] or ".")
+
+        # To be able to search projects by short-name we need a branch
+        branch = context.get_branch()
+        if not branch:
+            raise EnvironmentError(f"Unable to find a branch with a .uprojectdirs from '{context.get_primary_path()}'")
+
+        # Find the file that lists search paths
+        candidate = candidate.lower()
+        for i, uproj_path in enumerate(branch.read_projects()):
+            print("%-2d" % i, os.path.normpath(uproj_path.parent))
+            if uproj_path.stem.lower() == candidate:
+                print("Project found")
+                candidate = uproj_path
+                break
+        else:
+            self.print_error(f"Unable to find a project called '{candidate}' from", branch.get_dir())
+            return
+
+        return candidate
+
+    def main(self):
+        ret = self._main_impl()
+        self._provision(ret)
+        return ret
+
+    def _main_impl(self):
+        if not self.args.nameorpath:
+            if selected := self._select_project():
+                self.args.nameorpath = selected
+
+        candidate = self.args.nameorpath
+
+        if not candidate:
+            self.print_error("No project name or path specified")
+            return False
+
+        if candidate.lower() == "active":
+            return self._print_active()
+
+        if candidate.lower() == "auto":
+            candidate = self._set_auto()
+            if not candidate:
+                self.print_error(f"Unable to locate any projects from '{os.getcwd()}'")
+                return False
+
+        # Work out what the candidate should be
+        if candidate.lower() == "cwd":
+            return self._set_cwd()
+        elif candidate.lower() == "none" or candidate.lower() == "branch":
+            candidate = self._set_branch()
+        elif "/" in candidate or "\\" in candidate or os.path.exists(candidate):
+            candidate = self._set_qualified(candidate)
+        else:
+            candidate = self._set_short_form(candidate)
+
+        if not candidate:
+            return False
+
+        # Fire up the context
+        context = unreal.Context(candidate)
+        if project := context.get_project():
+            self.print_info("Changing active project")
+            print("Name:", project.get_name())
+            print("Path:", project.get_path())
+            primary_dir = str(project.get_path())
+        else:
+            primary_dir = str(context.get_primary_path())
+
+        # Record the project as active
+        session = self.get_noticeboard(self.Noticeboard.SESSION)
+        session["uproject"] = primary_dir
+
+        return primary_dir
+
+    def _provision(self, primary_dir):
+        if not primary_dir:
+            return
+
+        from pathlib import Path
+
+        # Find branch root
+        for candidate in Path(__file__).parents:
+            if (candidate / "ushell.bat").is_file():
+                root_dir = candidate
+                break
+        else:
+            return
+
+        # Find the bundle of data to provision
+        input_path = root_dir / Path("channels/bundle")
+        if not input_path.is_file():
+            return
+
+        with input_path.open("rb") as inp:
+            input_data = memoryview(inp.read())
+
+        input_mtime = input_path.stat().st_mtime
+
+        # Check we've something to provision from
+        nyckels = self._get_nyckels(primary_dir)
+        if not nyckels:
+            return
+
+        self.print_info("Provisioning")
+
+        # Let's provision!
+        import cipher
+        blob = cipher.Blob(input_data)
+        for nyckel in nyckels:
+            for name, data in blob.find(nyckel):
+                path = root_dir / name
+                if not path.is_file():
+                    print(name)
+                    dest_path = (root_dir / path)
+                    dest_path.parent.mkdir(parents=True, exist_ok=True)
+                    with dest_path.open("wb") as out:
+                        out.write(data)
+
+    def _get_nyckels(self):
+        pass
+        """
+        context = unreal.Context(primary_dir)
+        ini_paths = [x for x in context.glob("Extras/ushell/Settings.ini")]
+        if not ini_paths:
+            return
+
+        # And we need to know what we can provision
+        import base64
+        nyckels = []
+        for ini_path in ini_paths:
+            ini = unreal.Ini()
+            with ini_path.open("rt") as inp:
+                ini.load(inp)
+            if nyckel := ini.ushell.blob:
+                nyckel = base64.b64decode(str(nyckel))
+                nyckels.append(nyckel)
+        """
diff --git a/Engine/Extras/ushell/channels/unreal/core/cmds/run.py b/Engine/Extras/ushell/channels/unreal/core/cmds/run.py
new file mode 100644
--- /dev/null
+++ b/Engine/Extras/ushell/channels/unreal/core/cmds/run.py
@@ -0,0 +1,424 @@
+# Copyright Epic Games, Inc. All Rights Reserved.
+
+import os
+import sys
+import unreal
+import unrealcmd
+import unreal.cmdline
+from pathlib import Path
+
+
+#-------------------------------------------------------------------------------
+def _get_ip():
+    import socket
+    s = socket.socket(socket.AF_INET, socket.SOCK_DGRAM)
+    try:
+        s.connect(("172.31.255.255", 1))
+        ret = s.getsockname()[0]
+    except:
+        ret = "127.0.0.1"
+    finally:
+        s.close()
+    return ret
+
+
+
+#-------------------------------------------------------------------------------
+class _Attachable(object):
+    attach = unrealcmd.Opt(False, "Attach a debugger to the launched process")
+
+    def _get_debugger(self, name=None):
+        name = name or ("vs" if os.name == "nt" else "lldb")
+        ue_context = self.get_unreal_context()
+        channel = self.get_channel()
+        debuggers = channel.get_extendable("debuggers")
+
+        try:
+            debugger_class = debuggers.get_extension_class(name)
+            return debugger_class.construct(ue_context)
+        except ValueError:
+            raise ValueError(f"Unknown debugger '{name}'. Valid list: [{','.join([x[0] for x in debuggers.read_extensions()])}]")
+
+    def attach_to(self, pid, *, transport=None, host_ip=None):
+        debugger = self._get_debugger()
+        return debugger.attach(pid, transport, host_ip)
+
+    def debug(self, cmd, *args):
+        debugger = self._get_debugger()
+        exec_context = self.get_exec_context()
+        return debugger.debug(exec_context, cmd, *args)
+
+
+
+#-------------------------------------------------------------------------------
+class _RunCmd(unrealcmd.Cmd):
+    _variant = unrealcmd.Arg("development", "Build variant to launch")
+    runargs  = unrealcmd.Arg([str], "Arguments passed on to the process being launched")
+    build    = unrealcmd.Opt(False, "Build the target prior to running it")
+
+    def _build(self, target, variant, platform):
+        if target.get_type() == unreal.TargetType.EDITOR:
+            build_cmd = ("_build", "editor", variant)
+        else:
+            platform = platform or unreal.Platform.get_host()
+            build_cmd = ("_build", "target", target.get_name(), platform, variant)
+
+        import subprocess
+        ret = subprocess.run(build_cmd)
+        if ret.returncode:
+            raise SystemExit(ret.returncode)
+
+    def get_binary_path(self, target, platform=None):
+        ue_context = self.get_unreal_context()
+
+        if isinstance(target, str):
+            target = ue_context.get_target_by_name(target)
+        else:
+            target = ue_context.get_target_by_type(target)
+
+        if self.args.build:
+            self._build(target, self.args.variant, platform)
+
+        variant = unreal.Variant.parse(self.args.variant)
+        if build := target.get_build(variant, platform):
+            if build := build.get_binary_path():
+                return build
+
+        raise EnvironmentError(f"No {self.args.variant} build found for target '{target.get_name()}'")
+
+
+
+#-------------------------------------------------------------------------------
+class Editor(_RunCmd, _Attachable, unrealcmd.MultiPlatformCmd):
+    """ Launches the editor """
+    noproject   = unrealcmd.Opt(False, "Start the editor without specifying a project")
+    variant     = _RunCmd._variant
+
+    def main(self):
+        self.use_all_platforms()
+
+        ue_context = self.get_unreal_context()
+
+        binary_path = self.get_binary_path(unreal.TargetType.EDITOR)
+
+        args = (*unreal.cmdline.read_ueified(*self.args.runargs),)
+        if (project := ue_context.get_project()) and not self.args.noproject:
+            args = (project.get_path(), *args)
+
+        wants_terminal = "-stdout" in (x.lower() for x in self.args.runargs)
+        if wants_terminal and binary_path.suffix == "exe":
+            binary_path = binary_path[:-4]
+            binary_path += "-Cmd.exe"
+
+        if self.args.attach:
+            self.debug(binary_path, *args)
+            return
+
+        exec_context = self.get_exec_context()
+        editor = exec_context.create_runnable(str(binary_path), *args);
+        if wants_terminal or not self.is_interactive() or os.name != "nt":
+            import uelogprinter
+            printer = uelogprinter.Printer()
+            printer.run(editor)
+            return editor.get_return_code()
+        else:
+            editor.launch()
+
+
+
+#-------------------------------------------------------------------------------
+class Commandlet(_RunCmd, _Attachable, unrealcmd.MultiPlatformCmd):
+    """ Runs a commandlet """
+    commandlet = unrealcmd.Arg(str, "Name of the commandlet to run")
+    variant    = _RunCmd._variant
+
+    def complete_commandlet(self, prefix):
+        ue_context = self.get_unreal_context()
+        for dot_h in ue_context.glob("Source/Editor/UnrealEd/Classes/Commandlets/*.h"):
+            if "Commandlet.h" in dot_h.name:
+                yield dot_h.name[:-12]
+
+    def main(self):
+        self.use_all_platforms()
+
+        binary_path = self.get_binary_path(unreal.TargetType.EDITOR)
+        if binary_path.suffix == ".exe":
+            binary_path = binary_path.parent / (binary_path.stem + "-Cmd.exe")
+
+        args = (
+            "-run=" + self.args.commandlet,
+            *unreal.cmdline.read_ueified(*self.args.runargs),
+        )
+        if project := self.get_unreal_context().get_project():
+            args = (project.get_path(), *args)
+
+        if self.args.attach:
+            self.debug(str(binary_path), *args)
+            return
+
+        exec_context = self.get_exec_context()
+        cmd = exec_context.create_runnable(str(binary_path), *args);
+        if not self.is_interactive():
+            cmd.run()
+        else:
+            import uelogprinter
+            printer = uelogprinter.Printer()
+            printer.run(cmd)
+        return cmd.get_return_code()
+
+
+
+#-------------------------------------------------------------------------------
+class _ProgramTarget(_RunCmd, _Attachable, unrealcmd.MultiPlatformCmd):
+    def _main_impl(self):
+        self.use_all_platforms()
+
+        binary_path = self.get_binary_path(self.args.target)
+        args = (*unreal.cmdline.read_ueified(*self.args.runargs),)
+
+        if self.args.attach:
+            self.debug(str(binary_path), *args)
+            return
+
+        launch_kwargs = {}
+        if not sys.stdout.isatty():
+            launch_kwargs = { "silent" : True }
+        exec_context = self.get_exec_context()
+        runnable = exec_context.create_runnable(str(binary_path), *args);
+        runnable.launch(**launch_kwargs)
+        return True if runnable.is_gui() else runnable.wait()
+
+#-------------------------------------------------------------------------------
+class Program(_ProgramTarget):
+    """ Launches the specified program """
+    program = unrealcmd.Arg(str, "Name of the program to run")
+    variant = _RunCmd._variant
+
+    def complete_program(self, prefix):
+        ue_context = self.get_unreal_context()
+
+        for depth in ("*", "*/*"):
+            for target_cs in ue_context.glob(f"Source/Programs/{depth}/*.Target.cs"):
+                yield target_cs.name[:-10]
+
+    def main(self):
+        self.args.target = self.args.program
+        return self._main_impl()
+
+#-------------------------------------------------------------------------------
+class Target(_ProgramTarget):
+    """ Runs a name target on the host platform. """
+    target  = unrealcmd.Arg(str, "Name of the target to run")
+    variant = _RunCmd._variant
+
+    def complete_target(self, prefix):
+        #seen = set()
+        #for name, _ in self.get_unreal_context().read_targets():
+            #if name not in seen:
+                #seen.add(name)
+                #yield name
+        pass
+
+    def main(self):
+        return self._main_impl()
+
+
+
+#-------------------------------------------------------------------------------
+class _Runtime(_RunCmd, _Attachable):
+    """ When using --attach to debug to a process running on a devkit it is expected
+    that the Visual Studio solution is already open (i.e. '.sln open' has been run
+    prior to '.run ... --attach'). """
+    platform = unrealcmd.Arg("", "Platform to run on")
+    variant  = _RunCmd._variant
+    trace    = unrealcmd.Opt((False, ""), "Add command line argument to trace to development host")
+    cooked   = unrealcmd.Opt(False, "Use cooked data instead of staged data (local only)")
+    onthefly = unrealcmd.Opt(False, "Add the -filehostip= argument.")
+    binpath  = unrealcmd.Opt("", "Override the binary that is launched")
+    datadir  = unrealcmd.Opt("", "Use and alternative staged/packaged directory")
+
+    def _main_impl(self, target_type):
+        ue_context = self.get_unreal_context()
+
+        platform = self.args.platform
+        if not platform:
+            platform = unreal.Platform.get_host()
+
+        platform = self.get_platform(platform)
+        platform_name = platform.get_name()
+
+        project = ue_context.get_project()
+        if not project:
+            raise EnvironmentError(f"An active project is required to launch a {target_type.name.lower()} instance")
+
+        cook_form = platform.get_cook_form(target_type.name.lower())
+        stage_dir = project.get_dir() / f"Saved/StagedBuilds/{cook_form}"
+
+        # Allow the user to override the data directory used.
+        if self.args.datadir:
+            if self.args.cooked:
+                self.print_warning("Ignoring --cooked because --datadir was given")
+                self.args.cooked = False
+
+            stage_dir = Path(self.args.datadir).absolute()
+            if not stage_dir.is_dir():
+                self.print_error(f"Directory {stage_dir} does not exist")
+                return False
+
+            if not (stage_dir / "Engine").is_dir():
+                self.print_warning(f"An Engine/ directory is expected in '{stage_dir}'")
+
+
+        # I probably shouldn't do this, but some platforms need this file to be
+        # able to run with data served from the host PC. But I want to be explicit
+        # and clear what command line arguments are given
+
+        # Find the file in a manner that maintains case.
+        cmdline_txt = None
+        maybe_names = (
+            "UECommandLine.txt",
+            "uecommandline.txt",
+            "UE4CommandLine.txt",
+            "ue4commandline.txt",
+        )
+        for maybe_name in maybe_names:
+            cmdline_txt = stage_dir / maybe_name
+            if not cmdline_txt.is_file():
+                continue
+            cmdline_txt = cmdline_txt.resolve() # corrects case on Windows
+
+            dest_txt = cmdline_txt.parent / (cmdline_txt.stem + "_old_ushell.txt")
+            if not dest_txt.is_file():
+                cmdline_txt.rename(dest_txt)
+
+            break
+        else:
+            cmdline_txt = None
+
+        if cmdline_txt is not None:
+            cmdline_txt.open("wb").close()
+
+        # Get the path to the binary to execute
+        if self.args.binpath:
+            binary_path = Path(self.args.binpath).resolve()
+            if not binary_path.is_file():
+                raise ValueError(f"Binary not found; '{binary_path}'")
+        else:
+            binary_path = self.get_binary_path(target_type, platform_name)
+
+        # The user can opt-in to using cooked data instead of staged data
+        if self.args.cooked:
+            stage_dir = None
+            if platform.get_name() != unreal.Platform.get_host():
+                raise ValueError(f"The '--cooked' option may only be used when launching on {unreal.Platform.get_host()}")
+
+        # Validate data directories
+        if stage_dir:
+            if not stage_dir.is_dir():
+                suffix = "Client" if target_type == unreal.TargetType.CLIENT else ""
+                stage_dir = project.get_dir() / f"Saved/StagedBuilds/{platform_name}{suffix}/"
+
+            if not stage_dir.is_dir():
+                raise EnvironmentError(f"No staged data found for '{cook_form}'")
+        else:
+            cook_dir = project.get_dir() / f"Saved/Cooked/{cook_form}/"
+            if not cook_dir.is_dir():
+                sandboxed = next((x for x in self.args.runargs if "-sandbox=" in x), None)
+                if sandboxed:
+                    self.print_warning(f"No cooked data found for '{cook_form}'")
+                else:
+                    raise EnvironmentError(f"No cooked data found for '{cook_form}'")
+
+        # Condition command line with -filehostip
+        has_filehostip = False
+        for arg in self.args.runargs:
+            if "-filehostip" in arg.lower():
+                has_filehostip = True
+                break
+        else:
+            if self.args.onthefly:
+                self.args.runargs = (*self.args.runargs, "-filehostip=" + _get_ip())
+                has_filehostip = True
+
+        # The project argument.
+        project_arg = project.get_name()
+        if not has_filehostip:
+            project_arg = f"../../../{project_arg}/{project_arg}.uproject"
+
+        # Build arguments
+        args = (
+            project_arg,
+            *unreal.cmdline.read_ueified(*self.args.runargs),
+        )
+
+        if not self.args.cooked:
+            args = (*args, "-pak")
+
+        if self.args.trace is not False:
+            args = (*args, "-tracehost=" + _get_ip())
+            if self.args.trace:
+                args = (*args, "-trace=" + self.args.trace)
+
+        # Launch the runtime
+        stage_dir = str(stage_dir) + "/" if stage_dir else ""
+        exec_context = self.get_exec_context()
+        instance = platform.launch(exec_context, stage_dir, str(binary_path), args)
+        if not instance:
+            raise RuntimeError("An error occurred while launching.")
+
+        # Attach a debugger
+        if self.args.attach:
+            if type(instance) is tuple:
+                attach_ok = False
+                pid, devkit_ip = instance
+
+                # If there is no devkit involved we can try and use
+                if devkit_ip == None:
+                    attach_ok = self.attach_to(pid)
+
+                # Otherwise we're using VS to attach and that's only on Windows
+                elif vs_transport := getattr(platform, "vs_transport", None):
+                    print("Attaching to", pid, "on", devkit_ip, "with", vs_transport)
+                    attach_ok = self.attach_to(pid, transport=vs_transport, host_ip=devkit_ip)
+
+                # Tell the user stuff
+                if attach_ok:
+                    print("Debugger attached successfully")
+                else:
+                    self.print_warning("Unable to attach a debugger")
+
+                return
+
+            self.print_warning(f"--attach is not currently supported on {platform.get_name()}")
+
+#-------------------------------------------------------------------------------
+class Game(_Runtime):
+    """ Launches the game runtime. """
+    def complete_platform(self, prefix):
+        yield from super().complete_platform(prefix)
+        yield "editor"
+
+    def main(self):
+        if self.args.platform == "editor":
+            editor = Editor()
+            editor._channel = self._channel
+            args = (self.args.variant, "--", *self.args.runargs, "-game")
+            if self.args.build:  args = ("--build", *args)
+            if self.args.attach: args = ("--attach", *args)
+            return editor.invoke(args)
+
+        return self._main_impl(unreal.TargetType.GAME)
+
+#-------------------------------------------------------------------------------
+class Client(_Runtime):
+    """ Launches the game client. """
+    def main(self):
+        return self._main_impl(unreal.TargetType.CLIENT)
+
+#-------------------------------------------------------------------------------
+class Server(_Runtime):
+    """ Launches the server binary. """
+    complete_platform = ("win64", "linux")
+    def main(self):
+        self.args.runargs = (*self.args.runargs, "-log", "-unattended")
+        return self._main_impl(unreal.TargetType.SERVER)
