Updated ushell [26875160]

diff --git a/Engine/Extras/ushell/channels/unreal/core/cmds/run_include_tool.py b/Engine/Extras/ushell/channels/unreal/core/cmds/run_include_tool.py
new file mode 100644
--- /dev/null
+++ b/Engine/Extras/ushell/channels/unreal/core/cmds/run_include_tool.py
@@ -0,0 +1,81 @@
+# Copyright Epic Games, Inc. All Rights Reserved.
+
+import unreal
+import tempfile
+import unrealcmd
+import subprocess
+
+#-------------------------------------------------------------------------------
+def _build_include_tool(exec_context, ue_context):
+    # Work out where MSBuild is
+    vswhere = (
+        "vswhere.exe",
+        "-latest",
+        "-find MSBuild/*/Bin/amd64/MSBuild.exe",
+    )
+    proc = subprocess.Popen(" ".join(vswhere), stdout=subprocess.PIPE)
+    msbuild_exe = next(iter(proc.stdout.readline, b""), b"").decode().rstrip()
+    proc.stdout.close()
+
+    msbuild_exe = msbuild_exe or "MSBuild.exe"
+
+    # Get MSBuild to build IncludeTool
+    engine = ue_context.get_engine()
+    csproj_path = engine.get_dir() / "Source/Programs/IncludeTool/IncludeTool/IncludeTool.csproj"
+    msbuild_args = (
+        "/target:Restore;Build",
+        "/property:Configuration=Release,Platform=AnyCPU",
+        "/v:m",
+        str(csproj_path)
+    )
+    return exec_context.create_runnable(msbuild_exe, *msbuild_args).run() == 0
+
+#-------------------------------------------------------------------------------
+def _run_include_tool(exec_context, ue_context, target, variant, *inctool_args):
+    temp_dir = tempfile.TemporaryDirectory()
+    engine = ue_context.get_engine()
+    it_bin_path = engine.get_dir() / "Binaries/DotNET/IncludeTool/IncludeTool.exe"
+    if not it_bin_path.is_file():
+        it_bin_path = engine.get_dir() / "Binaries/DotNET/IncludeTool.exe"
+    it_args = (
+        "-Mode=Scan",
+        "-Target=" + target,
+        "-Platform=Linux",
+        "-Configuration=" + variant,
+        "-WorkingDir=" + temp_dir.name,
+        *inctool_args,
+    )
+    return exec_context.create_runnable(str(it_bin_path), *it_args).run()
+
+#-------------------------------------------------------------------------------
+class RunIncludeTool(unrealcmd.MultiPlatformCmd):
+    """ Runs IncludeTool.exe """
+    target      = unrealcmd.Arg(str, "The target to run the IncludeTool on")
+    variant     = unrealcmd.Arg("development", "Build variant to be processed (default=development)")
+    inctoolargs = unrealcmd.Arg([str], "Arguments to pass through to IncludeTool")
+
+    def complete_target(self, prefix):
+        return ("editor", "game", "client", "server")
+
+    def main(self):
+        self.use_platform("linux")
+        exec_context = self.get_exec_context()
+        ue_context = self.get_unreal_context()
+
+        # Establish the target to use
+        try:
+            target = unreal.TargetType[self.args.target.upper()]
+            target = ue_context.get_target_by_type(target).get_name()
+        except:
+            target = self.args.target
+
+        # Build
+        self.print_info("Building IncludeToole")
+        if not _build_include_tool(exec_context, ue_context):
+            self.print_error("Failed to build IncludeTool")
+            return False
+
+        # Run
+        variant = self.args.variant
+        self.print_info("Running on", target + "/" + variant)
+        return _run_include_tool(exec_context, ue_context, target, variant, *self.args.inctoolargs)
diff --git a/Engine/Extras/ushell/channels/unreal/core/cmds/sln.py b/Engine/Extras/ushell/channels/unreal/core/cmds/sln.py
new file mode 100644
--- /dev/null
+++ b/Engine/Extras/ushell/channels/unreal/core/cmds/sln.py
@@ -0,0 +1,121 @@
+# Copyright Epic Games, Inc. All Rights Reserved.
+
+import os
+import shutil
+import unrealcmd
+
+#-------------------------------------------------------------------------------
+class _Base(object):
+    def _get_primary_name(self):
+        name = "UE4"
+        try:
+            ue_context = self.get_unreal_context()
+            engine = ue_context.get_engine()
+
+            ppn_path = engine.get_dir() / "Intermediate/ProjectFiles/PrimaryProjectName.txt"
+            if not ppn_path.is_file():
+                # Old path, needs to be maintained while projects prior to UE5.1 are supported
+                ppn_path = engine.get_dir() / "Intermediate/ProjectFiles/MasterProjectName.txt"
+            with ppn_path.open("rb") as ppn_file:
+                name = ppn_file.read().decode().strip()
+        except FileNotFoundError:
+            pass
+        return name
+
+    def _get_sln_path(self):
+        ue_context = self.get_unreal_context()
+        engine = ue_context.get_engine()
+        primary_name = self._get_primary_name()
+        return engine.get_dir().parent / (primary_name + ".sln")
+
+    def _open_sln(self):
+        if os.name != "nt":
+            self.print_error("Opening project files is only supported on Windows")
+            return False
+
+        sln_path = self._get_sln_path()
+
+        self.print_info("Enumerating VS instances")
+        import vs.dte
+        for i, instance in enumerate(vs.dte.running()):
+            if instance_path := instance.get_sln_path():
+                print(i, instance_path, end="")
+                if sln_path.samefile(instance_path):
+                    if instance.activate():
+                        print(" ...activating")
+                        return True
+                print()
+
+        os.chdir(os.path.dirname(sln_path))
+
+        run_args = ("cmd.exe", "/c", "start", sln_path)
+        shell_open = self.get_exec_context().create_runnable(*run_args)
+        self.print_info("Opening project")
+        print("Path:", sln_path)
+        return shell_open.run()
+
+#-------------------------------------------------------------------------------
+class Generate(unrealcmd.MultiPlatformCmd, _Base):
+    """ Generates a Visual Studio solution for the active project. To more
+    easily distinguish open solutions from one branch to the next, the generated
+    .sln file will be suffixed with the current branch name. Use the '--all'
+    option to include all the current branch's projects. To generate-then-open
+    invoke as '.sln generate open'."""
+    open    = unrealcmd.Arg("", "If this argument is 'open' then open after generating")
+    ubtargs = unrealcmd.Arg([str], "Additional arguments pass to UnrealBuildTool")
+    notag   = unrealcmd.Opt(False, "Do not tag the solution name with a branch identifier")
+    all     = unrealcmd.Opt(False, "Include all the branch's projects")
+
+    def complete_open(self, prefix):
+        yield "open"
+
+    def get_exec_context(self):
+        context = super().get_exec_context()
+        if not self.args.notag:
+            context.get_env()["UE_NAME_PROJECT_AFTER_FOLDER"] = "1"
+        return context
+
+    def main(self):
+        self.use_all_platforms()
+        ue_context = self.get_unreal_context()
+
+        exec_context = self.get_exec_context()
+        args = ("-ProjectFiles", *self.args.ubtargs)
+        if not self.args.all:
+            if project := ue_context.get_project():
+                args = ("-Project=" + str(project.get_path()), *args)
+
+        self.print_info("Generating project files")
+        ubt = ue_context.get_engine().get_ubt()
+        for cmd, args in ubt.read_actions(*args):
+            cmd = exec_context.create_runnable(cmd, *args)
+            ret = cmd.run()
+            if ret:
+                return ret
+
+        if self.args.open == "open":
+            return self._open_sln()
+
+#-------------------------------------------------------------------------------
+class Open(unrealcmd.MultiPlatformCmd, _Base):
+    """ Opens project files in Visual Studio """
+
+    def get_exec_context(self):
+        context = super().get_exec_context()
+
+        # Detect when UE_NAME_PROJECT_AFTER_FOLDER was used to generate the
+        # project files and set it to allow tools like UnrealVS to generate
+        # the project files consistently.
+        primary_parts = self._get_primary_name().split("_", 1)
+        if len(primary_parts) > 1:
+            ue_context = self.get_unreal_context()
+            engine = ue_context.get_engine()
+            folder_name = engine.get_dir().parent.parent.name
+            if folder_name == primary_parts[1]:
+                context.get_env()["UE_NAME_PROJECT_AFTER_FOLDER"] = "1"
+
+        return context
+
+    def main(self):
+        self.use_all_platforms()
+        return self._open_sln()
diff --git a/Engine/Extras/ushell/channels/unreal/core/cmds/stage.py b/Engine/Extras/ushell/channels/unreal/core/cmds/stage.py
new file mode 100644
--- /dev/null
+++ b/Engine/Extras/ushell/channels/unreal/core/cmds/stage.py
@@ -0,0 +1,160 @@
+# Copyright Epic Games, Inc. All Rights Reserved.
+
+import unreal
+import unrealcmd
+import subprocess
+
+#-------------------------------------------------------------------------------
+class _Impl(unrealcmd.Cmd):
+    target   = unrealcmd.Arg(str, "Cooked target to stage")
+    platform = unrealcmd.Arg("", "Platform whose cook should be staged")
+    variant  = unrealcmd.Arg("development", "Build variant to stage binaries to copy")
+    uatargs  = unrealcmd.Arg([str], "Additional arguments to pass to UAT")
+
+    complete_target = ("game", "client", "server")
+
+    def _build(self, build_editor, ubt_args=None):
+        self.print_info("Building")
+
+        monolithic = (
+            self.args.target,
+            self.args.platform,
+            self.args.variant,
+        )
+        if ubt_args:
+            import shlex
+            monolithic = (
+                *monolithic,
+                "--",
+                *shlex.split(ubt_args),
+            )
+
+        schedule = []
+        if build_editor:
+            schedule.append(("editor", ))
+        schedule.append(monolithic)
+
+        for params in schedule:
+            ret = subprocess.run(("_build", *params))
+            if ret.returncode:
+                raise SystemExit(ret.returncode)
+
+    def _cook(self, cook_args=None):
+        self.print_info("Cooking")
+
+        if cook_args:
+            import shlex
+            cook_args = ("--", *shlex.split(cook_args))
+
+        ret = subprocess.run((
+            "_cook",
+            self.args.target,
+            self.args.platform,
+            *(cook_args if cook_args else tuple()),
+        ))
+
+        if ret.returncode:
+            raise SystemExit(ret.returncode)
+
+    @unrealcmd.Cmd.summarise
+    def main(self, *, skipstage=False):
+        # Check there is a valid active project
+        ue_context = self.get_unreal_context()
+        project = ue_context.get_project()
+        if not project:
+            raise EnvironmentError("An active project is required to stage a build")
+
+        # Fix up the platform argument so it has a value
+        platform = self.args.platform
+        if not platform:
+            platform = unreal.Platform.get_host()
+        platform = self.args.platform = platform.lower()
+
+        # Validate the given target
+        target = unreal.TargetType.parse(self.args.target)
+        target = ue_context.get_target_by_type(target)
+        if not target:
+            self.print_error(f"'{project.get_name()}' does not appear to have a '{self.args.target}' target")
+            return 1
+
+        # Inform the user about the context stage is going to operate in.
+        self.print_info("Context")
+        print("Target:", self.args.target)
+        print("Platform:", self.args.platform)
+        print("Variant:", self.args.variant)
+
+        # Build and cook convenience
+        if self.args.build is not False:
+            ubt_args = self.args.build if self.args.build else None
+            self._build(self.args.cook is not False, ubt_args=ubt_args)
+            print()
+
+        if self.args.cook is not False:
+            cook_args = self.args.cook if self.args.cook else None
+            self._cook(cook_args)
+            print()
+
+        # Launch UAT
+        self.print_info("Staging")
+
+        if target.get_type() == unreal.TargetType.GAME:
+            target_arg = None
+        elif target.get_type() == unreal.TargetType.CLIENT:
+            target_arg = "-client"
+        elif target.get_type() == unreal.TargetType.SERVER:
+            target_arg = "-server"
+            self.args.uatargs = (*self.args.uatargs, "-noclient")
+        else:
+            raise ValueError(f"Unable to stage target of type '{self.args.target}'")
+
+        platform = self.get_platform(platform)
+
+        dot_uat_cmd = "_uat"
+        dot_uat_args = (
+            "BuildCookRun",
+            "--", # NB. NOT a UAT arg!
+            #"-project=" + ue_context.get_project().get_path() <-- implicit via _uat
+            target_arg,
+            "-skipbuild" if "-build" not in self.args.uatargs else None,
+            "-skipcook" if "-cook" not in self.args.uatargs else None,
+            "-skipstage" if skipstage else "-stage",
+            "-deploy" if self.args.deploy else None,
+            "-config=" + self.args.variant,
+            "-platform=" + platform.get_name(),
+            None if self.args.nopak else "-pak",
+            *self.args.uatargs,
+        )
+
+        exec_context = self.get_exec_context()
+        cmd = exec_context.create_runnable(dot_uat_cmd, *dot_uat_args)
+        return cmd.run()
+
+
+
+#-------------------------------------------------------------------------------
+class Stage(_Impl):
+    """ Stages a build from the latest cook and built binaries. Note that
+    specifying the variant is only relevant if this stage is to be deployed at
+    some point in the future (it is not the binary used when using '.run').
+
+    The --build/cook options accept an optional argument that are the arguments
+    passed on to UnrealBuildTool and the cook commandlet respectively. For
+    example; '--build="-GoSuperFast -DisableUnity"'. Double-up ("") or escape
+    (\\") to pass through quotes."""
+    build  = unrealcmd.Opt((False, ""), "Build code prior to staging")
+    cook   = unrealcmd.Opt((False, ""), "Run a cook before running the stage step")
+    deploy = unrealcmd.Opt(False, "Deploy to devkit/device after staging")
+    nopak  = unrealcmd.Opt(False, "Staged result should be loose files")
+
+
+
+#-------------------------------------------------------------------------------
+class Deploy(_Impl):
+    """ Deploys a staged build to a devkit for the given platform. If there is no
+    existing staging to deploy the command '.stage ... --deploy' can be used. """
+    def main(self):
+        self.args.build  = False
+        self.args.cook   = False
+        self.args.deploy = True
+        self.args.nopak  = True
+        return super().main(skipstage=True)
diff --git a/Engine/Extras/ushell/channels/unreal/core/cmds/uat.py b/Engine/Extras/ushell/channels/unreal/core/cmds/uat.py
new file mode 100644
--- /dev/null
+++ b/Engine/Extras/ushell/channels/unreal/core/cmds/uat.py
@@ -0,0 +1,68 @@
+# Copyright Epic Games, Inc. All Rights Reserved.
+
+import os
+import unrealcmd
+
+#-------------------------------------------------------------------------------
+class Uat(unrealcmd.MultiPlatformCmd):
+    """ Runs an UnrealAutomationTool command """
+    command = unrealcmd.Arg(str, "The UnrealAutomationTool command to run")
+    uatargs = unrealcmd.Arg([str], "Arguments to pass to the UAT command")
+
+    def complete_command(self, prefix):
+        import os
+        import re
+
+        def _read_foreach_subdir(dir, func):
+            for entry in (x for x in os.scandir(dir) if x.is_dir()):
+                yield from func(entry.path)
+
+        def _read_buildcommand_types(dir, depth=0):
+            for entry in os.scandir(dir):
+                if entry.name.endswith(".cs"):
+                    with open(entry.path, "rt") as lines:
+                        for line in lines:
+                            m = re.search("(\w+) : BuildCommand", line)
+                            if m:
+                                yield m.group(1)
+                                break
+                elif entry.is_dir() and entry.name != "obj" and depth < 2:
+                    yield from _read_buildcommand_types(entry.path, depth + 1)
+                    continue
+
+        def _find_automation_rules(dir):
+            for entry in os.scandir(dir):
+                if entry.name.endswith(".Automation.csproj"):
+                    yield from _read_buildcommand_types(dir)
+                    return
+
+        ue_context = self.get_unreal_context()
+
+        uat_script_roots = (
+            ue_context.get_engine().get_dir() / "Source/Programs/AutomationTool",
+        )
+
+        if project := ue_context.get_project():
+            uat_script_roots = (
+                *uat_script_roots,
+                project.get_dir() / "Build",
+            )
+
+        for script_root in uat_script_roots:
+            yield from _read_foreach_subdir(script_root, _find_automation_rules)
+
+    def main(self):
+        self.use_all_platforms()
+
+        ue_context = self.get_unreal_context()
+
+        cmd = ue_context.get_engine().get_dir() / "Build/BatchFiles"
+        cmd /= "RunUAT.bat" if os.name == "nt" else "RunUAT.sh"
+        args = (self.args.command,)
+
+        if project := ue_context.get_project():
+            args = (*args, "-project=" + str(project.get_path()))
+
+        exec_context = self.get_exec_context()
+        cmd = exec_context.create_runnable(cmd, *args, *self.args.uatargs)
+        return cmd.run()
diff --git a/Engine/Extras/ushell/channels/unreal/core/debuggers/lldb.py b/Engine/Extras/ushell/channels/unreal/core/debuggers/lldb.py
new file mode 100644
--- /dev/null
+++ b/Engine/Extras/ushell/channels/unreal/core/debuggers/lldb.py
@@ -0,0 +1,27 @@
+# Copyright Epic Games, Inc. All Rights Reserved.
+
+import os
+import signal
+import unreal
+import subprocess
+
+#-------------------------------------------------------------------------------
+class Debugger(unreal.Debugger):
+    name = "lldb"
+
+    def _disable_ctrl_c(self):
+        def nop_handler(*args): pass
+        self._prev_sigint_handler = signal.signal(signal.SIGINT, nop_handler)
+
+    def _restore_ctrl_c(self):
+        signal.signal(signal.SIGINT, self._prev_sigint_handler)
+
+    def _debug(self, exec_context, cmd, *args):
+        self._disable_ctrl_c()
+        subprocess.run(("lldb", cmd, "--", *args))
+        self._restore_ctrl_c()
+
+    def _attach(self, pid, transport=None, host_ip=None):
+        self._disable_ctrl_c()
+        subprocess.run(("lldb", "--attach-pid", str(pid)))
+        self._restore_ctrl_c()
diff --git a/Engine/Extras/ushell/channels/unreal/core/debuggers/vs.py b/Engine/Extras/ushell/channels/unreal/core/debuggers/vs.py
new file mode 100644
--- /dev/null
+++ b/Engine/Extras/ushell/channels/unreal/core/debuggers/vs.py
@@ -0,0 +1,46 @@
+# Copyright Epic Games, Inc. All Rights Reserved.
+
+import os
+import unreal
+import subprocess
+from pathlib import Path
+
+#-------------------------------------------------------------------------------
+class Debugger(unreal.Debugger):
+    name = "Visual Studio"
+
+    def _debug(self, exec_context, cmd, *args):
+        cmd = exec_context.create_runnable(cmd, *args)
+        cmd.launch(suspended=True, new_term=True)
+        pid = cmd.get_pid()
+
+        try:
+            if not self._attach(pid):
+                cmd.kill()
+        except:
+            cmd.kill()
+            raise
+
+    def _attach(self, pid, transport=None, host_ip=None):
+        engine = self.get_unreal_context().get_engine()
+        engine_dir = engine.get_dir()
+
+        import vs.dte
+        for instance in vs.dte.running():
+            candidate = instance.get_sln_path()
+            if candidate and engine_dir.is_relative_to(Path(candidate).parent):
+                if instance.attach(pid, transport, host_ip):
+                    break
+        else:
+            # There's nothing to fallback if a fancy debugger is supposed to be used
+            if transport:
+                return False
+
+            ret = subprocess.run(("vsjitdebugger.exe", "-p", str(pid)))
+            if ret.returncode:
+                return False
+
+        if not transport:
+            vs.dte.resume_process(pid)
+
+        return True
diff --git a/Engine/Extras/ushell/channels/unreal/core/describe.flow.py b/Engine/Extras/ushell/channels/unreal/core/describe.flow.py
new file mode 100644
--- /dev/null
+++ b/Engine/Extras/ushell/channels/unreal/core/describe.flow.py
@@ -0,0 +1,245 @@
+# Copyright Epic Games, Inc. All Rights Reserved.
+
+import sys
+import flow.describe
+
+#-------------------------------------------------------------------------------
+fzf_linux = flow.describe.Tool()
+fzf_linux.version("0.33.0")
+fzf_linux.platform("linux")
+fzf_linux.payload("[REDACTED_PROFILE_URL]/fzf/releases/download/0.33.0/fzf-0.33.0-linux_amd64.tar.gz")
+fzf_linux.bin("fzf")
+fzf_linux.sha1("41aeb1c7bc01c700614768a777c5df91919a180c")
+
+fzf_darwin = flow.describe.Tool()
+fzf_darwin.version("0.33.0")
+fzf_darwin.platform("darwin")
+fzf_darwin.payload("[REDACTED_PROFILE_URL]/fzf/releases/download/0.33.0/fzf-0.33.0-darwin_amd64.zip")
+fzf_darwin.bin("fzf")
+fzf_darwin.sha1("f83505967ca3d0ad849a461029d60cd414e3ed76")
+
+
+
+#-------------------------------------------------------------------------------
+ripgrep_win32 = flow.describe.Tool()
+ripgrep_win32.version("13.0.0")
+ripgrep_win32.platform("win32")
+ripgrep_win32.payload("[REDACTED_PROFILE_URL]/ripgrep/releases/download/13.0.0/ripgrep-13.0.0-x86_64-pc-windows-msvc.zip")
+ripgrep_win32.bin("rg.exe")
+ripgrep_win32.sha1("f54ffef82c50d6079f211ed5348a81aa98274881")
+ripgrep_win32.source("[REDACTED_PROFILE_URL]/ripgrep/releases/latest", "ripgrep-(\d+\.\d+\.\d+)-x86_64")
+
+ripgrep_linux = flow.describe.Tool()
+ripgrep_linux.version("13.0.0")
+ripgrep_linux.platform("linux")
+ripgrep_linux.payload("[REDACTED_PROFILE_URL]/ripgrep/releases/download/13.0.0/ripgrep-13.0.0-x86_64-unknown-linux-musl.tar.gz")
+ripgrep_linux.bin("rg")
+ripgrep_linux.sha1("eae550a5ed0fcef8bc555c9ec41c593b8b5aea46")
+
+ripgrep_darwin = flow.describe.Tool()
+ripgrep_darwin.version("13.0.0")
+ripgrep_darwin.platform("darwin")
+ripgrep_darwin.payload("[REDACTED_PROFILE_URL]/ripgrep/releases/download/13.0.0/ripgrep-13.0.0-x86_64-apple-darwin.tar.gz")
+ripgrep_darwin.bin("rg")
+ripgrep_darwin.sha1("f09500e866c49c641ecb42fd8a0089d664d7d385")
+
+
+
+#-------------------------------------------------------------------------------
+build_target = flow.describe.Command()
+build_target.source("cmds/build.py", "Build")
+build_target.invoke("build", "target")
+
+for target_type in ("Editor", "Program", "Server", "Client", "Game"):
+    cmd = flow.describe.Command()
+    cmd.source("cmds/build.py", target_type)
+    cmd.invoke("build", target_type.lower())
+    globals()["build_" + target_type] = cmd
+
+    cmd = flow.describe.Command()
+    cmd.source("cmds/build.py", "Clean" + target_type)
+    cmd.invoke("build", "clean", target_type.lower())
+    globals()["build_clean" + target_type] = cmd
+
+    del cmd # so as not to pollute global scope with flow commands
+
+#-------------------------------------------------------------------------------
+build_xml_show = flow.describe.Command()
+build_xml_show.source("cmds/build_xml.py", "Show")
+build_xml_show.invoke("build", "xml")
+
+build_xml_edit = flow.describe.Command()
+build_xml_edit.source("cmds/build_xml.py", "Edit")
+build_xml_edit.invoke("build", "xml", "edit")
+
+build_xml_set = flow.describe.Command()
+build_xml_set.source("cmds/build_xml.py", "Set")
+build_xml_set.invoke("build", "xml", "set")
+
+build_xml_clear = flow.describe.Command()
+build_xml_clear.source("cmds/build_xml.py", "Clear")
+build_xml_clear.invoke("build", "xml", "clear")
+
+#-------------------------------------------------------------------------------
+build_clangdb = flow.describe.Command()
+build_clangdb.source("cmds/clangdb.py", "ClangDb")
+build_clangdb.invoke("build", "misc", "clangdb")
+
+#-------------------------------------------------------------------------------
+run_editor = flow.describe.Command()
+run_editor.source("cmds/run.py", "Editor")
+run_editor.invoke("run", "editor")
+
+run_commandlet = flow.describe.Command()
+run_commandlet.source("cmds/run.py", "Commandlet")
+run_commandlet.invoke("run", "commandlet")
+
+run_program = flow.describe.Command()
+run_program.source("cmds/run.py", "Program")
+run_program.invoke("run", "program")
+
+run_server = flow.describe.Command()
+run_server.source("cmds/run.py", "Server")
+run_server.invoke("run", "server")
+
+run_client = flow.describe.Command()
+run_client.source("cmds/run.py", "Client")
+run_client.invoke("run", "client")
+
+run_game = flow.describe.Command()
+run_game.source("cmds/run.py", "Game")
+run_game.invoke("run", "game")
+
+run_target = flow.describe.Command()
+run_target.source("cmds/run.py", "Target")
+run_target.invoke("run", "target")
+
+#-------------------------------------------------------------------------------
+run_include_tool = flow.describe.Command()
+run_include_tool.source("cmds/run_include_tool.py", "RunIncludeTool")
+run_include_tool.invoke("run", "includetool")
+
+#-------------------------------------------------------------------------------
+kill = flow.describe.Command()
+kill.source("cmds/kill.py", "Kill")
+kill.invoke("kill")
+
+#-------------------------------------------------------------------------------
+cook = flow.describe.Command()
+cook.source("cmds/cook.py", "Cook")
+cook.invoke("cook")
+
+cook_server = flow.describe.Command()
+cook_server.source("cmds/cook.py", "Server")
+cook_server.invoke("cook", "server")
+
+cook_client = flow.describe.Command()
+cook_client.source("cmds/cook.py", "Client")
+cook_client.invoke("cook", "client")
+
+cook_game = flow.describe.Command()
+cook_game.source("cmds/cook.py", "Game")
+cook_game.invoke("cook", "game")
+
+#-------------------------------------------------------------------------------
+ddc_auth = flow.describe.Command()
+ddc_auth.source("cmds/ddc.py", "Auth")
+ddc_auth.invoke("ddc", "auth")
+
+#-------------------------------------------------------------------------------
+sln_generate = flow.describe.Command()
+sln_generate.source("cmds/sln.py", "Generate")
+sln_generate.invoke("sln", "generate")
+
+sln_open = flow.describe.Command()
+sln_open.source("cmds/sln.py", "Open")
+sln_open.invoke("sln", "open")
+
+#-------------------------------------------------------------------------------
+uat = flow.describe.Command()
+uat.source("cmds/uat.py", "Uat")
+uat.invoke("uat")
+
+#-------------------------------------------------------------------------------
+stage = flow.describe.Command()
+stage.source("cmds/stage.py", "Stage")
+stage.invoke("stage")
+
+#-------------------------------------------------------------------------------
+deploy = flow.describe.Command()
+deploy.source("cmds/stage.py", "Deploy")
+deploy.invoke("deploy")
+
+#-------------------------------------------------------------------------------
+info = flow.describe.Command()
+info.source("cmds/info.py", "Info")
+info.invoke("info")
+
+#-------------------------------------------------------------------------------
+info_projects = flow.describe.Command()
+info_projects.source("cmds/info.py", "Projects")
+info_projects.invoke("info", "projects")
+
+#-------------------------------------------------------------------------------
+info_config = flow.describe.Command()
+info_config.source("cmds/info.py", "Config")
+info_config.invoke("info", "config")
+
+#-------------------------------------------------------------------------------
+notify = flow.describe.Command()
+notify.source("cmds/notify.py", "Notify")
+notify.invoke("notify")
+
+#-------------------------------------------------------------------------------
+project_change = flow.describe.Command()
+project_change.source("cmds/project.py", "Change")
+project_change.invoke("project")
+
+#-------------------------------------------------------------------------------
+prompt = flow.describe.Command()
+prompt.source("prompt.py", "Prompt")
+prompt.invoke("prompt")
+prompt.prefix("$")
+
+#-------------------------------------------------------------------------------
+boot = flow.describe.Command()
+boot.source("boot.py", "Boot")
+boot.invoke("boot")
+boot.prefix("$")
+
+#-------------------------------------------------------------------------------
+tips = flow.describe.Command()
+tips.source("tips.py", "Tips")
+tips.invoke("tip")
+tips.prefix("$")
+
+
+
+#-------------------------------------------------------------------------------
+vswhere = flow.describe.Tool()
+vswhere.version("3.0.1")
+vswhere.platform("win32")
+vswhere.payload("[REDACTED_PROFILE_URL]/vswhere/releases/download/3.0.1/vswhere.exe")
+vswhere.bin("vswhere.exe")
+vswhere.sha1("209d10bd07c6dc032ad26b40e4d93227c7c385fb")
+vswhere.source("[REDACTED_PROFILE_URL]/vswhere/releases", "(\d+\.\d+\.\d+)/vswhere.exe")
+
+
+
+#-------------------------------------------------------------------------------
+def _add_debugger(name):
+    extension = flow.describe.Extension()
+    extension.source(f"debuggers/{name}.py", "Debugger")
+    extension.mount("debuggers", name)
+    globals()[name] = extension
+
+_add_debugger("lldb")
+if sys.platform == "win32":
+    _add_debugger("vs")
+
+
+
+#-------------------------------------------------------------------------------
+unreal = flow.describe.Channel()
+unreal.version("0")
+unreal.extendable("debuggers")
diff --git a/Engine/Extras/ushell/channels/unreal/core/prompt.py b/Engine/Extras/ushell/channels/unreal/core/prompt.py
new file mode 100644
--- /dev/null
+++ b/Engine/Extras/ushell/channels/unreal/core/prompt.py
@@ -0,0 +1,62 @@
+# Copyright Epic Games, Inc. All Rights Reserved.
+
+import os
+import unreal
+
+#-------------------------------------------------------------------------------
+class Prompt(object):
+    def prompt(self, context):
+        try:
+            self._get_ue_branch(context)
+        except EnvironmentError:
+            self._get_git_branch(context)
+
+        context.TITLE = context.UE_BRANCH + " | " + context.UE_PROJECT
+        super().prompt(context)
+
+    def _get_ue_branch(self, context):
+        try:
+            session = self.get_noticeboard(self.Noticeboard.SESSION)
+            ue_context = unreal.Context(session["uproject"] or ".")
+
+            branch_full = ue_context.get_engine().get_info().get("BranchName", "UnknownBranch")
+            branch_name = ue_context.get_name()
+            if project := ue_context.get_project():
+                project_name = project.get_name()
+
+                if ue_context.get_type() == unreal.ContextType.FOREIGN:
+                    branch_name = ue_context.get_engine().get_dir().parent.name
+                    project_name += "(F)"
+            else:
+                project_name = "-"
+        except EnvironmentError:
+            raise
+        except Exception as e:
+            branch_full = "Error"
+            branch_name = "Error"
+            project_name = str(e) + "\nError"
+
+        context.UE_BRANCH = branch_name
+        context.UE_PROJECT = project_name
+        context.UE_BRANCH_FULL = branch_full
+
+    def _get_git_branch(self, context):
+        from pathlib import Path
+        for git_root in (Path(os.getcwd()) / "x").parents:
+            if (git_root / ".git").is_dir():
+                break
+        else:
+            context.UE_BRANCH = "NoBranch"
+            context.UE_PROJECT = "NoProject"
+            return
+
+        context.UE_BRANCH = git_root.stem
+        context.UE_PROJECT = "nobranch"
+        with (git_root / ".git/HEAD").open("rt") as file:
+            for line in file.readlines():
+                line = line.strip()
+                if line.startswith("ref: "):
+                    context.UE_PROJECT = line[5:].split("/")[-1]
+                else:
+                    context.UE_PROJECT = line[:6]
+                break
diff --git a/Engine/Extras/ushell/channels/unreal/core/pylib/cipher.py b/Engine/Extras/ushell/channels/unreal/core/pylib/cipher.py
new file mode 100644
--- /dev/null
+++ b/Engine/Extras/ushell/channels/unreal/core/pylib/cipher.py
@@ -0,0 +1,160 @@
+# Copyright Epic Games, Inc. All Rights Reserved.
+
+import secrets
+from pathlib import Path
+from typing import Tuple, List, Iterator
+
+#-------------------------------------------------------------------------------
+def _rol(x:int, n:int) -> int:
+    x &= 0xffffffff
+    x = (x << n) | (x >> (32 - n))
+    return x
+
+#-------------------------------------------------------------------------------
+def _qr(x:List[int], a:int, b:int, c:int, d:int) -> None:
+    x[a] += x[b]; x[d] = _rol(x[d] ^ x[a], 16)
+    x[c] += x[d]; x[b] = _rol(x[b] ^ x[c], 12)
+    x[a] += x[b]; x[d] = _rol(x[d] ^ x[a],  8)
+    x[c] += x[d]; x[b] = _rol(x[b] ^ x[c],  7)
+
+#-------------------------------------------------------------------------------
+def _read_blocks(key:bytes, nonce:bytes, counter:int=0) -> Iterator[bytes]:
+    key = memoryview(key).cast("L")
+    nonce = memoryview(nonce).cast("L")
+    state = [0x61707865, 0x3320646e, 0x79622d32, 0x6b206574, *key, 0, 0, *nonce]
+    out = memoryview(bytearray(4 * 16)).cast("L")
+    while True:
+        state[12] = counter
+        state[13] = counter >> 32
+        counter += 1
+
+        block = list(state)
+
+        for i in range(10):
+            _qr(block, 0, 4,  8, 12); _qr(block, 1, 5,  9, 13)
+            _qr(block, 2, 6, 10, 14); _qr(block, 3, 7, 11, 15)
+            _qr(block, 0, 5, 10, 15); _qr(block, 1, 6, 11, 12)
+            _qr(block, 2, 7,  8, 13); _qr(block, 3, 4,  9, 14)
+
+        for i in range(len(block)):
+            out[i] = (block[i] + state[i]) & 0xffffffff
+
+        yield from (x for x in out.cast("B"))
+
+#-------------------------------------------------------------------------------
+def _chacha20(data:bytes, *, key:memoryview, nonce:memoryview):
+    assert len(key) == 32
+    assert len(nonce) == 8
+    return bytes(a ^ b for a, b in zip(data, _read_blocks(key, nonce)))
+
+#-------------------------------------------------------------------------------
+def _hash(x):
+    ret = 5381
+    for c in x:
+        ret = ((ret * 33) + c) & 0xffffffff
+    return ret
+
+
+
+#-------------------------------------------------------------------------------
+class Blob(object):
+    _data : memoryview
+    _lede : memoryview
+    _toc  : memoryview
+
+    def __init__(self, data:memoryview|None=None) -> None:
+        if not data:
+            data = bytearray(secrets.token_bytes(32 << 10))
+            data = memoryview(data)
+        self._set_data(data)
+        self._ranges = []
+        self._toc_set = [0] * (len(self._toc) // 128)
+
+    def _set_data(self, data):
+        self._data = data.cast("B")
+        self._lede = data[:256]
+        self._toc = data[256:4096]
+
+    def get_data(self) -> memoryview:
+        return self._data
+
+    def _diffuse(self, key:memoryview) -> memoryview:
+        index = key[0] % (len(self._lede) - 48)
+        data = _chacha20(self._lede[index:index + 48], key=key, nonce=b'\0' * 8)
+        data = memoryview(data)
+        return data[:32], data[32:40], data[40:48]
+
+    def find(self, key:memoryview) -> Iterator[Tuple[str, memoryview]]:
+        key, nonce, nonce_toc = self._diffuse(key)
+        for i in range(0, len(self._toc), 128):
+            piece = self._toc[i:i + 128]
+            piece = _chacha20(piece, key=key, nonce=nonce_toc)
+            piece = memoryview(piece).cast("L")
+            name = piece[4:].cast("B")
+            if piece[2] == _hash(name):
+                name = name[:piece[3]]
+
+                offset = piece[0]
+                size = piece[1]
+                data = self._data[offset:offset + size]
+                data = _chacha20(data, key=key, nonce=nonce)
+
+                yield bytes(name).decode(), data
+
+    def add(self, item:Path, key:memoryview, toc_name=None) -> None:
+        size = item.stat().st_size
+
+        # find somewhere to store it or grow if we fail a lot
+        offset = 0
+        space = len(self._data) - size - 4096
+        assert space > 0
+        for i in range(100):
+            offset = secrets.randbelow(space)
+            offset += 4096
+            for l,r in self._ranges:
+                if not (offset + size <= l or offset >= r):
+                    offset = 0
+                    break
+            else:
+                self._ranges.append((offset, offset + size))
+                break
+        else:
+            data = bytearray(self._data)
+            data += bytearray(secrets.token_bytes(8 << 10))
+            data = memoryview(data)
+            self._set_data(data)
+            return self.add(item, key, toc_name)
+
+        key, nonce, nonce_toc = self._diffuse(key)
+
+        # encypt the file
+        with item.open("rb") as inp:
+            data = _chacha20(inp.read(), key=key, nonce=nonce)
+
+        # create a toc entry
+        name = (toc_name or str(item)).encode()
+        name_length = len(name)
+        assert name_length <= 112
+        name = name + secrets.token_bytes(112 - name_length)
+
+        toc_entry = bytearray(128)
+        toc_entry[16:] = name
+        toc_entry = memoryview(toc_entry).cast("L")
+        toc_entry[0] = offset
+        toc_entry[1] = size
+        toc_entry[2] = _hash(name)
+        toc_entry[3] = name_length
+        toc_entry = toc_entry.cast("B")
+        toc_entry = _chacha20(toc_entry, key=key, nonce=nonce_toc)
+
+        # find somewhere to store the toc
+        while True:
+            toc_index = secrets.randbelow(len(self._toc_set))
+            if self._toc_set[toc_index] == 0:
+                self._toc_set[toc_index] = 1
+                break
+
+        # blit encrypted parts into blob
+        toc_index = toc_index * 128
+        self._data[offset:offset + size] = data
+        self._toc[toc_index:toc_index + 128] = toc_entry
diff --git a/Engine/Extras/ushell/channels/unreal/core/pylib/fzf.py b/Engine/Extras/ushell/channels/unreal/core/pylib/fzf.py
new file mode 100644
--- /dev/null
+++ b/Engine/Extras/ushell/channels/unreal/core/pylib/fzf.py
@@ -0,0 +1,41 @@
+# Copyright Epic Games, Inc. All Rights Reserved.
+
+import subprocess
+
+#-------------------------------------------------------------------------------
+def run(item_iter, *, height=10, prompt=">", multi=False, ansi=False, sort=True, ext_stdin=None):
+    fzf_args = (
+        "--prompt=" + prompt if prompt else None,
+        None if sort else "--no-sort",
+        "--multi" if multi == True else "--multi=" + str(multi) if multi > 0 else None,
+        "--height=" + str(height + 2), # "+2" because fzf adds two lines of prompt
+        "--color=16",
+        "--layout=reverse",
+        "--no-mouse",
+        "--ansi" if ansi else None,
+    )
+    proc = subprocess.Popen(
+        ("fzf", *(x for x in fzf_args if x)),
+        universal_newlines=True,
+        stdout=subprocess.PIPE,
+        stdin=ext_stdin if ext_stdin else subprocess.PIPE,
+    )
+
+    try:
+        if not ext_stdin:
+            try:
+                for item in item_iter:
+                    item = str(item) + "\n"
+                    proc.stdin.write(item)
+            except (BrokenPipeError, OSError):
+                pass
+
+            try: proc.stdin.close()
+            except: pass
+
+        yield from (x.strip() for x in proc.stdout.readlines())
+        proc.stdout.close()
+
+        proc.wait()
+    finally:
+        proc.kill()
diff --git a/Engine/Extras/ushell/channels/unreal/core/pylib/prettyprinter.py b/Engine/Extras/ushell/channels/unreal/core/pylib/prettyprinter.py
new file mode 100644
--- /dev/null
+++ b/Engine/Extras/ushell/channels/unreal/core/pylib/prettyprinter.py
@@ -0,0 +1,25 @@
+# Copyright Epic Games, Inc. All Rights Reserved.
+
+#-------------------------------------------------------------------------------
+class Printer(object):
+    def _thread(self):
+        try:
+            for line in iter(self._stdout.readline, b""):
+                self._print(line.decode(errors="replace").rstrip())
+        except (IOError, ValueError):
+            pass
+        finally:
+            self._stdout.close()
+
+    def run(self, runnable):
+        if not hasattr(self, "_print"):
+            name = type(self).__name__
+            raise SyntaxError("Printer '{name}' must implement a _print() method")
+
+        self._stdout = runnable.launch(stdout=True)
+
+        import threading
+        thread = threading.Thread(target=self._thread)
+        thread.start()
+        runnable.wait()
+        thread.join()
diff --git a/Engine/Extras/ushell/channels/unreal/core/pylib/uelogprinter.py b/Engine/Extras/ushell/channels/unreal/core/pylib/uelogprinter.py
new file mode 100644
--- /dev/null
+++ b/Engine/Extras/ushell/channels/unreal/core/pylib/uelogprinter.py
@@ -0,0 +1,43 @@
+# Copyright Epic Games, Inc. All Rights Reserved.
+
+import re
+import flow.cmd
+import prettyprinter
+
+#-------------------------------------------------------------------------------
+class Printer(prettyprinter.Printer):
+    def __init__(self):
+        self._log_re = re.compile(r"(\[.+\]|)Log(\w+):\s+(\w+):\s?")
+        self._decorators = (
+            flow.cmd.text.green,
+            flow.cmd.text.yellow,
+            flow.cmd.text.blue,
+            flow.cmd.text.magenta,
+            flow.cmd.text.cyan,
+            flow.cmd.text.light_green,
+            flow.cmd.text.light_yellow,
+            flow.cmd.text.light_blue,
+            flow.cmd.text.light_magenta,
+            flow.cmd.text.light_cyan,
+        )
+        self._decorations = {}
+
+    def _print_log(self, category, verbosity, line):
+        if verbosity == "Error":
+            line = flow.cmd.text.light_red(line)
+
+        decorator = self._decorations.get(category, None)
+        if not decorator:
+            hash = 5831
+            for c in category:
+                hash = ((hash << 5) + hash) ^ ord(c)
+            index = hash % len(self._decorators)
+            decorator = self._decorators[index]
+            self._decorations[category] = decorator
+
+        print(decorator(category), line)
+
+    def _print(self, line):
+        m = self._log_re.search(line)
+        if m: self._print_log(m.group(2), m.group(3), line[m.end():])
+        else: print(line)
diff --git a/Engine/Extras/ushell/channels/unreal/core/pylib/unreal/__init__.py b/Engine/Extras/ushell/channels/unreal/core/pylib/unreal/__init__.py
new file mode 100644
--- /dev/null
+++ b/Engine/Extras/ushell/channels/unreal/core/pylib/unreal/__init__.py
@@ -0,0 +1,7 @@
+# Copyright Epic Games, Inc. All Rights Reserved.
+
+from ._context import *
+from ._ini import *
+from ._platform import *
+from ._ubt import *
+from ._debugger import *
diff --git a/Engine/Extras/ushell/channels/unreal/core/pylib/unreal/_context.py b/Engine/Extras/ushell/channels/unreal/core/pylib/unreal/_context.py
new file mode 100644
--- /dev/null
+++ b/Engine/Extras/ushell/channels/unreal/core/pylib/unreal/_context.py
@@ -0,0 +1,680 @@
+# Copyright Epic Games, Inc. All Rights Reserved.
+
+import os
+import re
+import pathlib
+from enum import Enum, Flag
+from typing import ForwardRef, Iterator, Optional, Tuple, Dict
+
+from . import _ini as ini
+from ._ubt import BuildTool
+from ._platform import Platform
+
+# {{{1 misc --------------------------------------------------------------------
+
+#-------------------------------------------------------------------------------
+def _resolve_association_nt(association:str) -> Optional[pathlib.Path]:
+    import winreg
+
+    if association[0] == "{":
+        hive = winreg.HKEY_CURRENT_USER
+        key_path = r"SOFTWARE\Epic Games\Unreal Engine\Builds"
+        value_name = association
+    else:
+        hive = winreg.HKEY_LOCAL_MACHINE
+        key_path = r"SOFTWARE\EpicGames\Unreal Engine\\" + association
+        value_name = "InstalledDirectory"
+
+    try:
+        with winreg.OpenKey(hive, key_path) as key:
+            value = winreg.QueryValueEx(key, value_name)
+    except OSError:
+        return
+
+    if value[1] == 1:
+        return pathlib.Path(value[0])
+
+def _resolve_association_posix(association:str) -> Optional[pathlib.Path]:
+    raise NotImplementedError("Foreign projects are not yet supported")
+
+def _resolve_association(association:str) -> Optional[pathlib.Path]:
+    if os.name == "nt":
+        return _resolve_association_nt(association)
+    else:
+        return _resolve_association_posix(association)
+
+
+
+# {{{1 enums -------------------------------------------------------------------
+
+#-------------------------------------------------------------------------------
+class Variant(Enum):
+    DEBUG       = 0
+    DEBUGGAME   = 1
+    DEVELOPMENT = 2
+    TEST        = 3
+    SHIPPING    = 4
+
+    def ueify(self):
+        return "DebugGame" if self == Variant.DEBUGGAME else self.name.title()
+
+    @staticmethod
+    def parse(external:str) -> ForwardRef("Variant"):
+        try: return Variant[external.upper()]
+        except: raise ValueError(f"Unknown build variant '{external}'")
+
+#-------------------------------------------------------------------------------
+class TargetType(Enum):
+    UNKNOWN     = 0x00
+    PROGRAM     = 0x01
+    EDITOR      = 0x02
+    GAME        = 0x03
+    CLIENT      = 0x04
+    SERVER      = 0x05
+
+    @staticmethod
+    def parse(external:str) -> ForwardRef("TargetType"):
+        try: return TargetType[external.upper()]
+        except: raise ValueError(f"Unknown target type '{external}'")
+
+#-------------------------------------------------------------------------------
+class ContextType(Enum):
+    BRANCH      = 0 # context with no active project
+    PROJECT     = 1 # active project with Engine/ up the tree
+    FOREIGN     = 2 # active project with Engine/ located elsewhere (i.e. installed)
+
+
+
+# {{{1 ufs-item-----------------------------------------------------------------
+class _UfsItem(pathlib.Path):
+    _mount:ForwardRef("_UfsMount")
+
+    def __new__(cls, mount:ForwardRef("_UfsMount"), *args) -> ForwardRef("_UfsItem"):
+        cls = _WindowsUfsItem if os.name == "nt" else _PosixUfsItem
+        instance = super().__new__(cls, *args)
+        instance._mount = mount
+        return instance
+
+    def get_mount(self) -> ForwardRef("_UfsMount"):
+        return self._mount
+
+class _WindowsUfsItem(pathlib.WindowsPath, _UfsItem):
+    pass
+
+class _PosixUfsItem(pathlib.PosixPath, _UfsItem):
+    pass
+
+
+
+# {{{1 ufs_mount ---------------------------------------------------------------
+class _UfsMount(object):
+    _dir    :pathlib.Path
+    _node   :ForwardRef("_UfsNode")
+
+    def __init__(self, dir:pathlib.Path, node:ForwardRef("_UfsNode")):
+        self._dir = dir
+        self._node = node
+
+    def get_dir(self) -> pathlib.Path:
+        return self._dir
+
+    def get_node(self) -> ForwardRef("_UfsNode"):
+        return self._node
+
+    def glob(self, pattern:str) -> Iterator[_UfsItem]:
+        for x in self.get_dir().glob(pattern):
+            yield _UfsItem(self, x)
+
+    def get_item(self, file_path:str) -> Optional[_UfsItem]:
+        dest = self.get_dir() / file_path
+        if dest.is_file():
+            return _UfsItem(self, dest)
+
+
+
+# {{{1 build -------------------------------------------------------------------
+class _Build(object):
+    _receipt_path   :pathlib.Path
+    _variant        :Variant
+    _ufs_mount      :_UfsMount
+    _platform       :str
+
+    def __init__(self, receipt_path:pathlib.Path, variant:Variant, platform:str, ufs_mount:_UfsMount):
+        self._receipt_path = receipt_path
+        self._variant = variant
+        self._platform = platform
+        self._ufs_mount = ufs_mount
+
+    def get_variant(self) -> Variant:
+        return self._variant
+
+    def get_platform(self) -> str:
+        return self._platform
+
+    def get_binary_path(self) -> Optional[pathlib.Path]:
+        with self._receipt_path.open("rt") as x:
+            try:
+                import json
+                target = json.load(x)
+            except:
+                return
+
+        binary_path = target.get("Launch", None)
+        if not binary_path:
+            return
+
+        eng_node = proj_node = self._ufs_mount.get_node()
+        while parent_node := eng_node.get_parent_node():
+            eng_node = parent_node
+
+        binary_path = binary_path.replace("$(ProjectDir)", str(proj_node.get_dir()))
+        binary_path = binary_path.replace("$(EngineDir)", str(eng_node.get_dir()))
+
+        ret = pathlib.Path(binary_path)
+        return ret if ret.exists() else None
+
+
+
+# {{{1 target ------------------------------------------------------------------
+class _Target(object):
+    _type       :TargetType
+    _name       :str
+    _ufs_mount  :_UfsMount
+
+    def __init__(self, name:str, type:TargetType, ufs_mount:_UfsMount):
+        self._name = name
+        self._type = type
+        self._ufs_mount = ufs_mount
+
+    def get_name(self) -> str:
+        return self._name
+
+    def get_type(self) -> TargetType:
+        return self._type
+
+    def is_project_target(self) -> bool:
+        return isinstance(self._ufs_mount.get_node(), _Project)
+
+    def get_build(self, variant:Variant=Variant.DEVELOPMENT, platform:str="") -> Optional[_Build]:
+        platform = platform or Platform.get_host()
+
+        path = f"Binaries/{platform}/{self.get_name()}"
+        if variant != Variant.DEVELOPMENT:
+            path += "-" + platform + "-" + variant.ueify()
+        path += ".target"
+
+        if receipt := next(self._ufs_mount.glob(path), None):
+            return _Build(receipt, variant, platform, self._ufs_mount)
+
+
+
+# {{{1 ufs_node ----------------------------------------------------------------
+class _UfsNode(object):
+    _parent :ForwardRef("_UfsNode")
+    _dir    :pathlib.Path
+
+    def __init__(self, dir:pathlib.Path, parent:ForwardRef("_UfsNode")=None):
+        self._parent = parent
+        self._dir = dir
+
+    def get_parent_node(self) -> Optional[ForwardRef("_UfsNode")]:
+        return self._parent
+
+    def get_dir(self) -> pathlib.Path:
+        return self._dir
+
+    def get_root_mount(self) -> _UfsMount:
+        return _UfsMount(self._dir, self)
+
+    def read_mounts(self, *, ancestors:bool=True, shallow:bool=False, platforms:bool=True) -> Iterator[_UfsMount]:
+        if not shallow:
+            for dir in self._dir.glob("Restricted/*"):
+                yield _UfsMount(dir, self)
+                if platforms:
+                    yield from (_UfsMount(x, self) for x in dir.glob("Platforms/*"))
+
+            if platforms:
+                yield from (_UfsMount(x, self) for x in self._dir.glob("Platforms/*"))
+
+        yield self.get_root_mount()
+
+        if ancestors and self._parent:
+            yield from self._parent.read_mounts(ancestors=ancestors, shallow=shallow, platforms=platforms)
+
+    def glob(self, pattern:str, *, ancestors:bool=True, shallow:bool=False) -> Iterator[_UfsItem]:
+        for mount in self.read_mounts(ancestors=ancestors, shallow=shallow):
+            yield from mount.glob(pattern)
+
+    def _read_targets(self, targets_dir:str, **kwargs) -> Iterator[_Target]:
+        def parse_target_cs(target_cs):
+            def find_line(regex, line_iter):
+                for line in line_iter:
+                    m = re.search(regex, line)
+                    if m: return m.group(1)
+
+            with target_cs.open("rt") as cs:
+                line_iter = (x for x in cs)
+                name = find_line(r"class\s+(\S+)Target[\s:]", line_iter)
+                if name:
+                    type = find_line(r"=\s*TargetType\.(.+)\s*;", line_iter)
+                    return name, type or "UNKNOWN"
+
+        pattern = str(pathlib.Path(targets_dir) / "*.Target.cs")
+        for target_cs in self.glob(pattern, **kwargs):
+            if info := parse_target_cs(target_cs):
+                name, type = info
+                type = TargetType.parse(type)
+                yield _Target(name, type, target_cs.get_mount())
+
+    def get_target_by_type(self, target_type:TargetType) -> Optional[_Target]:
+        if target_type.value <= TargetType.PROGRAM.value:
+            return
+
+        for target in self._read_targets("Source"):
+            if target.get_type() == target_type:
+                return target
+
+    def get_target_by_name(self, target_name:str) -> Optional[_Target]:
+        lower_name = target_name.lower()
+
+        for dir in ("Source/Programs/*", "Source"):
+            for target in self._read_targets(dir):
+                if lower_name == target.get_name().lower():
+                    return target
+
+        node = self
+        while next_node := node.get_parent_node():
+            node = next_node
+
+        return _Target(target_name, TargetType.UNKNOWN, node.get_root_mount())
+
+
+
+# {{{1 config ------------------------------------------------------------------
+class _Config(object):
+    _ufs:_UfsNode
+
+    def __init__(self, ufs:_UfsNode):
+        self._ufs = ufs
+
+    def _load_hives(self, category:str, *, override:str="") -> ini.Ini:
+        ret = ini.Ini()
+
+        if not override:
+            hive_names = (
+                f"Default{category}.ini",
+                f"Base{category}.ini",
+            )
+            visit_platforms = False
+        else:
+            hive_names = (
+                f"{override}{category}.ini",
+                f"{override}{category}.ini",
+                f"Base{override}{category}.ini",
+            )
+            visit_platforms = True
+
+        node = self._ufs
+        for hive_name in iter(hive_names):
+            for mount in node.read_mounts(ancestors=False, platforms=visit_platforms):
+                if ufs_item := mount.get_item("Config/" + hive_name):
+                    with ufs_item.open("rt") as ini_file:
+                        ret.load(ini_file)
+
+            if next_node := node.get_parent_node():
+                node = next_node
+
+        return ret
+
+    def get(self, category:str, section:str="", key:str="", *, override:str="") -> Optional[ini.Ini|ini._Struct|ini._Value]:
+        ret = self._load_hives(category, override=override)
+
+        if section:
+            if (ret := getattr(ret, section, None)) and key:
+                ret = getattr(ret, key, None)
+
+        return ret
+
+
+
+# {{{1 engine ------------------------------------------------------------------
+class _Engine(_UfsNode):
+    def __init__(self, dir:pathlib.Path):
+        super().__init__(dir)
+
+    def get_ubt(self) -> BuildTool:
+        return BuildTool(self)
+
+    def get_info(self) -> dict:
+        ret = {
+            "MajorVersion" : self.get_version_major(),
+            "MinorVersion" : 0,
+            "PatchVersion" : 0,
+            "BranchName"   : "Error:Build.Version",
+            "Changelist"   : 0,
+        }
+
+        build_version_path = self._dir / "Build/Build.version"
+        if build_version_path.is_file():
+            with build_version_path.open() as x:
+                try:
+                    import json
+                    ret.update(json.load(x))
+                except:
+                    pass
+
+        return ret
+
+    def get_version_major(self) -> int:
+        key_path = self._dir / "Source/Developer/Windows/WindowsNoEditorTargetPlatform/WindowsNoEditorTargetPlatform.Build.cs"
+        return 4 if key_path.is_file() else 5
+
+    def get_version_full(self) -> tuple:
+        info = self.get_info()
+        return (info["MajorVersion"], info["MinorVersion"], info["PatchVersion"])
+
+
+
+# {{{1 project -----------------------------------------------------------------
+class _Project(_UfsNode):
+    _path   :pathlib.Path
+    _engine :_Engine
+
+    def __init__(self, path:pathlib.Path, engine:_Engine):
+        super().__init__(path.parent, engine)
+        self._path = path
+        self._engine = engine
+
+    def get_path(self) -> pathlib.Path:
+        return self._path
+
+    def get_name(self) -> pathlib.Path:
+        return self._path.stem
+
+    def get_engine(self) -> _Engine:
+        return self._engine
+
+    def get_target_by_type(self, target_type:TargetType, *, create_temp:bool=True) -> Optional[_Target]:
+        if target_type.value <= TargetType.PROGRAM.value:
+            return
+
+        for target in self._read_targets("Source", ancestors=False, shallow=True):
+            if target.get_type() == target_type:
+                return target
+
+        type_name = target_type.name.title()
+        target_name = self.get_name()
+        if target_type != TargetType.GAME:
+            target_name += type_name
+
+        target = _Target(target_name, target_type, self.get_root_mount())
+        if create_temp:
+            self._create_temp_target(target)
+
+        return target
+
+    def _create_temp_target(self, target:_Target) -> None:
+        int_dir = self.get_dir() / "Source"
+        int_dir.mkdir(parents=True, exist_ok=True)
+
+        type_name = target.get_type().name.title()
+        target_name = target.get_name()
+
+        dot_cs_path = int_dir / (target_name + ".Target.cs")
+        if dot_cs_path.is_file():
+            return
+
+        game_target = "UnrealGame"
+        if self._engine.get_version_major() <= 4:
+            game_target = "UE4Game"
+
+        with dot_cs_path.open("wt") as dot_cs:
+            dot_cs.write(
+                "// Created by ushell to make blueprint projects work end to end.\n"
+                "// Placed in Source/ because UAT cleans Intermediate/Source/\n"
+                "using UnrealBuildTool;\n"
+                "public class " + target_name + "Target : TargetRules {\n"
+                "public " + target_name + "Target(TargetInfo Target) : base(Target) {\n"
+                "Type = TargetType." + type_name + ";\n"
+                "ExtraModuleNames.AddRange(new string[] {\"" + game_target + "\"} ); }}\n"
+            )
+
+
+
+# {{{1 branch ------------------------------------------------------------------
+class _Branch(object):
+    _dir:pathlib.Path
+
+    def __init__(self, dir:pathlib.Path):
+        self._dir = dir
+
+    def get_name(self) -> str:
+        return self._dir.name
+
+    def get_dir(self) -> pathlib.Path:
+        return self._dir
+
+    def find_project(self, name:str) -> Optional[pathlib.Path]:
+        name = name.lower()
+        for project_path in self.read_projects():
+            if project_path.stem.lower() == name:
+                return project_path
+
+    def read_projects(self) -> Iterator[pathlib.Path]:
+        def read_dir(dir:pathlib.Path):
+            for sub_dir in (x for x in dir.glob("*") if x.is_dir()):
+                yield from sub_dir.glob("*.uproject")
+
+        def read_uprojectdirs(uprojdirs_path:pathlib.Path) -> Iterator[pathlib.Path]:
+            with uprojdirs_path.open() as file:
+                root_dir = uprojdirs_path.parent
+                for line in file:
+                    line = line.strip()
+                    for prefix in (";", "Templates", "Engine/"):
+                        if line.startswith(prefix):
+                            break
+                    else:
+                        yield from read_dir(root_dir / line)
+
+        for item in self._dir.glob("*.uprojectdirs"):
+            yield from read_uprojectdirs(item)
+
+
+
+# {{{1 platform-provider -------------------------------------------------------
+class _PlatformProvider(object):
+    _temp_ctx   :ForwardRef("Context")
+    _provisions :Dict[str, pathlib.Path]
+
+    def __init__(self, temp_ctx):
+        self._temp_ctx = temp_ctx
+        self._provisions = None
+
+    @classmethod
+    def _get_host_platform_name(self):
+        if hasattr(self, "_host_name"):
+            return self._host_name
+
+        if os.name == "nt":
+            self._host_name = "win64"
+        else:
+            import sys
+            if sys.platform == "linux": self._host_name = "linux"
+            elif sys.platform == "darwin": self._host_name = "mac"
+            else: raise EnvironmentError("Unknown host platform '{sys.platform}'")
+
+        return self._host_name
+
+    def _get_provisions(self) -> Dict[str, pathlib.Path]:
+        if self._provisions is not None:
+            return self._provisions
+
+        provisions = self._provisions = {}
+
+        # Built-in platforms
+        host_name = self._get_host_platform_name()
+        for item in pathlib.Path(__file__).parent.glob(f"platforms/{host_name}/*.py"):
+            if not item.name.startswith("_"):
+                provisions[item.stem] = item
+
+        if host_name != "win64":
+            return provisions
+
+        # Use Engine/Platforms if we are hosted in a UE branch
+        for candidate in pathlib.Path(__file__).parent.parent.parents:
+            if candidate.name == "Engine":
+                candidate = (candidate / "Platforms")
+                if candidate.is_dir():
+                    break
+        else:
+            return provisions
+
+        for platform_dir in (x for x in candidate.glob("*") if x.is_dir()):
+            extras_dir = platform_dir / "Extras/ushell"
+            if not extras_dir.is_dir():
+                continue
+
+            for item in extras_dir.glob("platform_*.py"):
+                provisions[item.stem[9:]] = item
+
+        return provisions
+
+    def get_platform(self, name:None|str=None) -> Platform|None:
+        host_name = self._get_host_platform_name()
+        name = (name or host_name).lower()
+
+        try:
+            path = f"unreal.platforms.{host_name}.{name}"
+            py_path = self._get_provisions()[name]
+            import importlib.util
+            spec = importlib.util.spec_from_file_location(path , py_path)
+            module = importlib.util.module_from_spec(spec)
+            spec.loader.exec_module(module)
+        except (KeyError, ModuleNotFoundError):
+            return
+
+        return module.Platform(self._temp_ctx)
+
+    def read_platform_names(self) -> Iterator[str]:
+        yield from self._get_provisions().keys()
+
+# {{{1 context -----------------------------------------------------------------
+class Context(object):
+    _type          :ContextType
+    _engine_dir    :pathlib.Path
+    _project_path  :pathlib.Path
+    _platforms     :_PlatformProvider
+
+    def __init__(self, primary_dir:str):
+        primary_dir:pathlib.Path = pathlib.Path(primary_dir).resolve()
+
+        # If a .uproject* was given instead the strip it off.
+        if primary_dir.suffix.startswith(".uproject"):
+            if primary_dir.is_file():
+                return self._initialize(primary_dir)
+            raise EnvironmentError(f"Project file not found '{primary_dir}'")
+
+        # Walk up from 'primary_dir' and find something to build a context from
+        if primary_dir.is_dir():
+            iter_dir:pathlib.Path = primary_dir
+            while not iter_dir.samefile(iter_dir.anchor):
+                if uproj_path := next(iter_dir.glob("*.uproject*"), None):
+                    return self._initialize(uproj_path)
+                iter_dir = iter_dir.parent
+            raise EnvironmentError(f"Unable to establish an Unreal context from directory '{primary_dir}'")
+
+        raise EnvironmentError(f"Unable to establish an Unreal context with '{primary_dir}'")
+
+    def _initialize(self, uproj_path:pathlib.Path) -> None:
+        if uproj_path.suffix == ".uproject":
+            return self._initialize_project(uproj_path)
+        else:
+            return self._initialize_noproject(uproj_path)
+
+    def _initialize_project(self, uproj_path:pathlib.Path) -> None:
+        self._project_path = uproj_path
+
+        # Find engine. First look up the tree from the project for an Engine/ dir
+        iter_dir:pathlib.Path = uproj_path.parent.parent
+        while not iter_dir.samefile(iter_dir.anchor):
+            candidate = iter_dir / "Engine"
+            if candidate.is_dir():
+                self._type = ContextType.PROJECT
+                self._engine_dir = candidate
+                return
+            iter_dir = iter_dir.parent
+
+        # Next inspect the .uproject file for evidence of an engine
+        uproj_data = None
+        with uproj_path.open("rt") as x:
+            try:
+                import json
+                uproj_data = json.load(x)
+            except:
+                pass
+
+        if uproj_data and (association := uproj_data.get("EngineAssociation")):
+            if engine_dir := _resolve_association(association):
+                engine_dir /= "Engine"
+                if engine_dir.is_dir():
+                    self._type = ContextType.FOREIGN
+                    self._engine_dir = engine_dir
+                    return
+
+        raise EnvironmentError(f"Unable to locate the engine for project '{uproj_path}'")
+
+    def _initialize_noproject(self, uprojdirs_path:pathlib.Path) -> None:
+        # Engine/ is expected to be a sibling of the .uprojectdirs file
+        engine_dir = uprojdirs_path.parent / "Engine"
+        if not engine_dir.is_dir():
+            raise EnvironmentError(f"Folder '{uprojdirs_path.parent}/Engine/' does not exist")
+
+        # We've everything for a valid context now
+        self._type = ContextType.BRANCH
+        self._engine_dir = engine_dir
+        self._project_path = None
+
+    def _get_ufs_node(self) -> _UfsNode:
+        return self.get_project() or self.get_engine()
+
+    def get_type(self) -> ContextType:
+        return self._type
+
+    def get_name(self) -> str:
+        return self._engine_dir.parent.name or "UnrealEngine"
+
+    def get_primary_path(self) -> pathlib.Path:
+        return self._project_path if self._project_path else self.get_branch().get_dir()
+
+    def get_project(self) -> Optional[_Project]:
+        if self._type != ContextType.BRANCH:
+            return _Project(self._project_path, self.get_engine())
+
+    def get_config(self) -> _Config:
+        return _Config(self._get_ufs_node())
+
+    def get_engine(self) -> _Engine:
+        return _Engine(self._engine_dir)
+
+    def get_branch(self, *, must_exist:bool=False) -> Optional[_Branch]:
+        if self._type != ContextType.FOREIGN:
+            return _Branch(self._engine_dir.parent)
+
+        if must_exist:
+            raise EnvironmentError(f"No valid branch found for '{self._engine_dir}'")
+
+    def get_platform_provider(self) -> _PlatformProvider:
+        if not hasattr(self, "_platforms"):
+            self._platforms = _PlatformProvider(self)
+        return self._platforms
+
+    def glob(self, pattern:str, shallow:bool=False) -> Iterator[_UfsItem]:
+        yield from self._get_ufs_node().glob(pattern, shallow=shallow)
+
+    def get_target_by_type(self, target_type:TargetType) -> Optional[_Target]:
+        return self._get_ufs_node().get_target_by_type(target_type)
+
+    def get_target_by_name(self, name:str) -> _Target:
+        return self._get_ufs_node().get_target_by_name(name)
+
+# vim: foldlevel=1 foldmethod=marker
diff --git a/Engine/Extras/ushell/channels/unreal/core/pylib/unreal/_debugger.py b/Engine/Extras/ushell/channels/unreal/core/pylib/unreal/_debugger.py
new file mode 100644
--- /dev/null
+++ b/Engine/Extras/ushell/channels/unreal/core/pylib/unreal/_debugger.py
@@ -0,0 +1,29 @@
+# Copyright Epic Games, Inc. All Rights Reserved.
+
+#-------------------------------------------------------------------------------
+def _optional_api(func):
+    def inner(self, *args, **kwargs):
+        if not hasattr(self, "_" + func.__name__):
+            raise NotImplementedError(f"Platform '{self.get_name()}' does not implement method '{func.__name__}()'")
+        return func(self, *args, **kwargs)
+
+    return inner
+
+#-------------------------------------------------------------------------------
+class Debugger(object):
+    def __init__(self, ue_context):
+        self._ue_context = ue_context
+
+    def get_unreal_context(self):
+        return self._ue_context
+
+    def get_name(self):
+        return type(self).name
+
+    @_optional_api
+    def debug(self, exec_context, cmd, *args):
+        return self._debug(exec_context, cmd, *args)
+
+    @_optional_api
+    def attach(self, pid, transport=None, host_ip=None):
+        return self._attach(pid, transport, host_ip)
diff --git a/Engine/Extras/ushell/channels/unreal/core/pylib/unreal/_ini.py b/Engine/Extras/ushell/channels/unreal/core/pylib/unreal/_ini.py
new file mode 100644
--- /dev/null
+++ b/Engine/Extras/ushell/channels/unreal/core/pylib/unreal/_ini.py
@@ -0,0 +1,123 @@
+# Copyright Epic Games, Inc. All Rights Reserved.
+
+import re
+
+#-------------------------------------------------------------------------------
+def _parse_struct(line):
+    if not line.startswith("("):
+        return line
+    ret = _Struct()
+    m = re.findall(r'(\w+)\s*=\s*(?:([^\s,"]+)|"(\\.|[^"]*)")', line[1:-1])
+    for key, value, temp in m:
+        ret.add_item(key, value or temp)
+    return ret
+
+#-------------------------------------------------------------------------------
+class _Value(object):
+    def __init__(self):
+        self._data = None
+        self._parsed = False
+
+    def get_inner(self):
+        self._parse()
+        return self._data
+
+    def _set(self, value):
+        self._data = value
+
+    def _append(self, value, dedupe):
+        if not isinstance(self._data, list):
+            self._data = []
+        if not dedupe or value not in self._data:
+            self._data.append(value)
+
+    def _parse(self):
+        if self._parsed or not self._data:
+            return
+        self._parsed = True
+        if isinstance(self._data, list):
+            self._data = [_parse_struct(x) for x in self._data]
+        else:
+            self._data = _parse_struct(self._data)
+
+    def __bool__(self):
+        return self._data is not None
+
+    def __repr__(self):
+        self._parse()
+        return str(self._data)
+
+    def __iter__(self):
+        self._parse()
+        if isinstance(self._data, list):
+            yield from self._data
+        elif self._data is not None:
+            yield self._data
+
+    def __getattr__(self, name):
+        self._parse()
+        if isinstance(self._data, _Struct):
+            return getattr(self._data, name) or _Value()
+        return _Value()
+
+#-------------------------------------------------------------------------------
+class _Struct(object):
+    def __init__(self):
+        self._data = {}
+
+    def get_inner(self):
+        return self._data
+
+    def add_item(self, key, value, cmd=""):
+        if cmd == "":
+            if key not in self._data:
+                self._data[key] = out = _Value()
+                out._set(value)
+            return
+
+        out = self._data.setdefault(key, _Value())
+        if   cmd == "+": out._append(value, True)
+        elif cmd == ".": out._append(value, False)
+
+    def __repr__(self):
+        return repr(self._data)
+
+    def __getattr__(self, name):
+        return self._data.get(name) or _Value()
+
+    def __iter__(self):
+        return iter(self._data.items())
+
+#-------------------------------------------------------------------------------
+class Ini(object):
+    def __init__(self):
+        self._sections = {}
+
+    def _add_section(self, name):
+        return self._sections.setdefault(name, _Struct())
+
+    @classmethod
+    def _lazy_re_compile(self):
+        if hasattr(self, "_lazied"):
+            return
+        self._comment_re = re.compile(r"\s*[#;]")
+        self._section_re = re.compile(r"\s*\[(.+)\]")
+        self._entry_re = re.compile(r"^\s*([+-.!@*]?)([^=]+)\s*=\s*(.+)\s*$")
+        self._lazied = True
+
+    def load(self, file):
+        self._lazy_re_compile()
+        section = None
+        for line in (x.strip() for x in file):
+            if m := Ini._comment_re.match(line):
+                continue
+            elif m := Ini._section_re.match(line):
+                section = self._add_section(m.group(1))
+            elif m := Ini._entry_re.match(line):
+                section.add_item(m.group(2), m.group(3), m.group(1))
+
+    def __getattr__(self, name):
+        return self._sections.get(name) or _Struct()
+
+    def __iter__(self):
+        return iter(self._sections.items())
diff --git a/Engine/Extras/ushell/channels/unreal/core/pylib/unreal/_platform.py b/Engine/Extras/ushell/channels/unreal/core/pylib/unreal/_platform.py
new file mode 100644
--- /dev/null
+++ b/Engine/Extras/ushell/channels/unreal/core/pylib/unreal/_platform.py
@@ -0,0 +1,162 @@
+# Copyright Epic Games, Inc. All Rights Reserved.
+
+import re
+import os
+import sys
+
+#-------------------------------------------------------------------------------
+def _optional_api(func):
+    def inner(self, *args, **kwargs):
+        if not hasattr(self, "_" + func.__name__):
+            raise NotImplementedError(f"Platform '{self.get_name()}' does not implement method '{func.__name__}()'")
+        return func(self, *args, **kwargs)
+
+    return inner
+
+
+
+#-------------------------------------------------------------------------------
+class Platform(object):
+    @staticmethod
+    def get_host():
+        if sys.platform == "linux":  return "Linux"
+        if sys.platform == "darwin": return "Mac"
+        if sys.platform == "win32":  return "Win64"
+        raise EnvironmentError("Unknown host platform; " + sys.platform)
+
+    @staticmethod
+    def get_sdks_dir():
+        env = os.getenv("UE_SDKS_ROOT")
+        if env:
+            return env + "/Host" + Platform.get_host() + "/"
+
+    def __init__(self, context):
+        self._context = context
+        self._version = None
+        self._sdks_dir = None
+
+    @staticmethod
+    def _get_version_helper_ue4(dot_cs, var_name):
+        needle = re.compile(var_name + r'\s+=\s+("|[\w.]+\("|)([^"]+)"?\)?\s*;')
+        try:
+            with open(dot_cs, "rt") as lines:
+                for line in (x for x in lines if var_name in x):
+                    m = needle.search(line)
+                    if m: return m.group(2)
+        except FileNotFoundError:
+            pass
+
+    def _get_version_helper_ue5(self, dot_cs, func_name="GetMainVersion"):
+        dot_cs = self._context.get_engine().get_dir() / dot_cs
+        try:
+            import re
+            with open(dot_cs, "rt") as cs_in:
+                lines = iter(cs_in)
+                next(x for x in lines if f"override string {func_name}" in x or "override string GetDesiredVersion()" in x)
+                next(lines) # {
+                for i in range(5):
+                    line = next(lines)
+                    if m := re.search(r'return "(.+)"', line):
+                        return m.group(1)
+        except (FileNotFoundError, StopIteration):
+            pass
+
+    def get_name(self):
+        return type(self).name
+
+    def get_config_name(self):
+        return getattr(type(self), "config_name", None) or self.get_name()
+
+    def get_autosdk_name(self):
+        return type(self).autosdk_name if hasattr(self, "autosdk_name") else self.get_name()
+
+    def get_unreal_context(self):
+        return self._context
+
+    def _get_version_json(self):
+        def conform_json(inp):
+            # As the time of writing the SDK.json files have trailing commas
+            # and C-style comments. So we'll try and conform it
+            ret = ""
+            count = 0
+            for line in inp:
+                line = line.strip()
+                if not line or line.startswith("//"):
+                    continue
+
+                count += 1
+                sep = "," if count > 1 else ""
+                if line.endswith(","):    line = line[:-1]
+                if line[0] in "[{":       count = 0
+                elif line[0] in "}]":     sep = ""
+                ret += sep + line
+            return ret
+
+        def read_sdk_json(path):
+            if not path.is_file():
+                return
+
+            import json
+            with path.open("rt") as inp:
+                data = conform_json(inp)
+                try:
+                    sdk_info = json.loads(data)
+                except:
+                    return
+
+            if version := sdk_info.get("AutoSDKDirectory", None):
+                return version
+
+            if version := sdk_info.get("MainVersion", None):
+                return version
+
+            if next_path := sdk_info.get("ParentSDKFile", None):
+                return read_sdk_json(path.parent / next_path)
+
+        name = self.get_config_name()
+        engine_dir = self._context.get_engine().get_dir()
+        json_path = engine_dir / f"Platforms/{name}/Config/{name}_SDK.json"
+        if version := read_sdk_json(json_path):
+            return version
+
+        json_path = engine_dir / f"Config/{name}/{name}_SDK.json"
+        return read_sdk_json(json_path)
+
+    def get_version(self):
+        if not (version := self._version or self._get_version_json()):
+            engine_version = self._context.get_engine().get_version_major()
+            self._version = self._get_version_ue5() if engine_version > 4 else self._get_version_ue4()
+            version = self._version or ""
+        return version
+
+    def read_env(self):
+        class Value(object):
+            def __init__(self, k, v):   self.key = k; self._value = v
+            def __str__(self):          return self.validate();
+            def __fspath__(self):       return str(self)
+            def get(self):              return self._value or "SDK_NOT_FOUND"
+            def validate(self):
+                value = self.get()
+                if not os.path.isdir(value):
+                    raise EnvironmentError(f"Invalid directory '{value}' for '{self.key}'")
+                return value
+
+        sdks_dir = Platform.get_sdks_dir() or ""
+        for key, value in self._read_env():
+            if value: value = os.path.join(sdks_dir, value)
+            yield Value(key, value)
+
+    @_optional_api
+    def get_cook_form(self, target):
+        form = self._get_cook_form(target)
+        if form:
+            return form
+        raise ValueError(f"No known cook form for target '{target}'")
+
+    @_optional_api
+    def launch(self, exec_context, stage_dir, binary_path, args):
+        return self._launch(exec_context, stage_dir, binary_path, args)
+
+    @_optional_api
+    def kill(self, target=None):
+        return self._kill(target)
diff --git a/Engine/Extras/ushell/channels/unreal/core/pylib/unreal/_ubt.py b/Engine/Extras/ushell/channels/unreal/core/pylib/unreal/_ubt.py
new file mode 100644
--- /dev/null
+++ b/Engine/Extras/ushell/channels/unreal/core/pylib/unreal/_ubt.py
@@ -0,0 +1,369 @@
+# Copyright Epic Games, Inc. All Rights Reserved.
+
+import os
+import re
+import sys
+import enum
+import tempfile
+from pathlib import Path
+
+#-------------------------------------------------------------------------------
+class _BuildConfigLevel(enum.Enum):
+    INTERNAL = 0
+    BRANCH   = 1
+    GLOBAL   = 2
+    LEGACY   = 3
+
+#-------------------------------------------------------------------------------
+class _BuildConfigXml(object):
+    def __init__(self):
+        self._depth = 0
+        self._comment_count = 0
+        self._root = {}
+        self._ns_skip = 0
+
+    def start(self, tag, attribs):
+        tag = tag[self._ns_skip + 2:] if tag.startswith("{") else tag
+        if self._depth == 1:   self._category = self._root.setdefault(tag, {})
+        elif self._depth == 2: self._prop = self._category.setdefault(tag, {"value":""})
+        self._depth += 1
+
+    def end(self, tag): self._depth -= 1
+    def start_ns(self, prefix, uri): self._ns_skip = len(uri)
+    def end_ns(self, prefix):        self._ns_skip = 0
+
+    def data(self, data):
+        if self._depth == 3:
+            if data := data.strip():
+                self._prop["value"] = data
+
+    def comment(self, data):
+        target = None
+        if self._depth == 1:   target = self._root
+        elif self._depth == 2: target = self._category
+        elif self._depth == 3: target = self._prop
+        if target != None:
+            x = f"$comment" + str(self._comment_count)
+            target[x] = data
+            self._comment_count += 1
+
+    def print(self, file=None):
+        indent = ""
+        def spam(*args):        print(indent, *args, sep="", file=file)
+        def spam_comment(data): spam("<!--", data, "-->")
+
+        spam("<?xml version='1.0' encoding='utf-8'?>")
+        spam(r'<Configuration xmlns="https://www.unrealengine.com/BuildConfiguration">')
+        indent = "  "
+        for category_name, category in self._root.items():
+            if category_name.startswith("$comment"):
+                spam_comment(category)
+                continue
+
+            if not len(category):
+                continue
+
+            spam(f"<{category_name}>")
+            indent = "    "
+            for prop_name, prop in category.items():
+                if prop_name.startswith("$comment"):
+                    spam_comment(prop)
+                    continue
+
+                for key, value in prop.items():
+                    if key.startswith("$comment"):
+                        spam_comment(value)
+                        continue
+
+                    spam(f"<{prop_name}>{value}</{prop_name}>")
+            indent = "  "
+            spam(f"</{category_name}>")
+        indent = ""
+        spam("</Configuration>")
+
+    def load(self, path):
+        import xml.etree.ElementTree as et
+        with open(path, "rb") as file:
+            parser = et.XMLParser(target=self)
+            try:
+                parser.feed(file.read())
+            except et.ParseError as e:
+                raise IOError(str(e))
+            parser.close()
+
+    def save(self, path):
+        with open(path, "wt") as file:
+            self.print(file=file)
+
+    def read_values(self):
+        for category, i in (x for x in self._root.items() if not x[0].startswith("$")):
+            for prop, j in (x for x in i.items() if not x[0].startswith("$")):
+                yield category, prop, j["value"]
+
+    def get_value(self, category, name):
+        if parent := self._root.get(category):
+            if prop := parent.get(name):
+                return prop["value"]
+
+        raise LookupError(f"No such property '{category}.{name}'")
+
+    def set_value(self, category, name, value):
+        category = self._root.setdefault(category, {})
+        prop = category.setdefault(name, {})
+        prop["value"] = value
+
+    def clear_value(self, category, name):
+        if parent := self._root.get(category):
+            if name in parent:
+                del parent[name]
+                return
+
+        raise LookupError(f"No such property '{category}.{name}'")
+
+#-------------------------------------------------------------------------------
+class _BuildConfig(object):
+    def __init__(self, path, level):
+        self._path = path
+        self._level = level
+        self._root = None
+
+    def _get_root(self):
+        if self._root:
+            return self._root
+
+        self._root = _BuildConfigXml()
+        if self.exists():
+            self._root.load(self._path)
+        return self._root
+
+    def get_path(self):                         return self._path
+    def get_level(self):                        return self._level
+    def exists(self):                           return os.path.isfile(self._path)
+    def save(self):                             self._get_root().save(self.get_path())
+    def read_values(self):                      yield from self._get_root().read_values()
+    def get_value(self, category, name):        return self._get_root().get_value(category, name)
+    def set_value(self, category, name, value): return self._get_root().set_value(category, name, value)
+    def clear_value(self, category, name):      return self._get_root().clear_value(category, name)
+
+#-------------------------------------------------------------------------------
+class _SchemaNode(object):
+    def __init__(self, root):
+        self._root = root
+
+    def get_name(self):
+        if isinstance(self._root, str): return self._root
+        else: return self._root.get("name")
+
+    def read_children(self):
+        if isinstance(self._root, str):
+            return
+
+        ns = { "x" : "http://www.w3.org/2001/XMLSchema" }
+        child = self._root.find("x:simpleType/x:restriction", ns)
+        if child:
+            yield from (_SchemaNode(x.get("value")) for x in child)
+            return
+
+        child = self._root.find("x:complexType/x:all", ns)
+        if child:
+            yield from (_SchemaNode(x) for x in child)
+            return
+
+        if self._root.get("type", None) == "boolean":
+            yield from (_SchemaNode("true"), _SchemaNode("false"))
+
+    def get_node(self, *ancestor_names):
+        node = self
+        for ancestor in ancestor_names:
+            for child in node.read_children():
+                if child.get_name() == ancestor:
+                    node = child
+                    break
+            else:
+                return
+
+        return node
+
+#-------------------------------------------------------------------------------
+class _BuildConfigSchema(_SchemaNode):
+    def __init__(self, path):
+        self._path = path
+
+        import xml.etree.ElementTree as et
+        xml = et.parse(self._path)
+        root = next(iter(xml.getroot()))
+        super().__init__(root)
+
+    def get_path(self):
+        return path
+
+
+
+#-------------------------------------------------------------------------------
+class _BuildToolBase(object):
+    def __init__(self, engine):
+        self._update_check = True
+        self._engine = engine
+
+    def _get_bin_path(self):
+        engine_dir = self._engine.get_dir()
+        for script, update_check in self._get_script_paths():
+            script = engine_dir / script
+            if script.is_file():
+                self._update_check = self._update_check and update_check
+                return script
+
+    def read_configurations(self):
+        engine_dir = self._engine.get_dir()
+
+        suffix = "UnrealBuildTool/BuildConfiguration.xml"
+        yield _BuildConfig(engine_dir / "Restricted/NotForLicensees/Programs" / suffix, _BuildConfigLevel.INTERNAL)
+        yield self.get_configuration(True)
+        yield self.get_configuration()
+        if os.name == "nt":
+            user_profile = os.getenv("USERPROFILE", "userprofile-notset")
+            yield _BuildConfig(user_profile + "/Documents/Unreal Engine/" + suffix, _BuildConfigLevel.LEGACY)
+        else:
+            config_dir = os.getenv("XDG_CONFIG_HOME", "~/.config")
+            config_dir = os.path.expanduser(config_dir)
+            yield _BuildConfig(config_dir + "/Unreal Engine/" + suffix, _BuildConfigLevel.LEGACY)
+
+    def get_configuration(self, branch=False, *, create=False):
+        suffix = "UnrealBuildTool/BuildConfiguration.xml"
+        if branch:
+            cfg_path = self._engine.get_dir() / "Saved" / suffix
+        else:
+            if os.name == "nt": prefix = os.getenv("APPDATA", "appdata")
+            else:               prefix = os.path.expanduser("~")
+            cfg_path = Path(prefix) / "Unreal Engine" / suffix
+
+        if create and not cfg_path.is_file():
+            cfg_path.parent.mkdir(parents=True, exist_ok=True)
+            with cfg_path.open("wt") as out:
+                print("<?xml version='1.0' encoding='utf-8'?>", file=out)
+                print(r'<Configuration xmlns="https://www.unrealengine.com/BuildConfiguration">', file=out)
+                print("</Configuration>", file=out)
+
+        level = _BuildConfigLevel.BRANCH if branch else _BuildConfigLevel.GLOBAL
+        return _BuildConfig(cfg_path, level)
+
+    def get_configuration_schema(self):
+        engine_dir = self._engine.get_dir()
+        xsd_path = engine_dir / "Saved/UnrealBuildTool/BuildConfiguration.Schema.xsd"
+        if xsd_path.is_file():
+            return _BuildConfigSchema(xsd_path)
+
+    def read_actions(self, *args):
+        cmd = self._get_bin_path()
+        if not (cmd and cmd.is_file()):
+            raise IOError("No UnrealBuildTool scripts found. Incomplete branch?")
+
+        if cmd.suffix == ".bat":
+            # cmd.exe insists on asking you if you want to terminate when Ctrl-C
+            # is pressed. Often it doesn't work or takes N goes. It seems that
+            # using devnull for stdin (<nul) makes it more tolerable. Messy :(
+            args = ("/d/s/c", "\"", cmd, *args, "<nul\"")
+            cmd = "cmd.exe"
+
+        if self._update_check:
+            self._update_check = False
+            yield from self._read_build_ubt_actions()
+
+        yield cmd, args
+
+#-------------------------------------------------------------------------------
+class _BuildToolWin64(_BuildToolBase):
+    def __init__(self, *args):
+        super().__init__(*args)
+        self._msbuild_response_path = None
+
+    def __del__(self):
+        if self._msbuild_response_path:
+            try: os.unlink(self._msbuild_response_path)
+            except: pass
+
+    def _read_build_ubt_actions(self):
+        # Find msbuild.exe that MS have buried away.
+        vswhere = (
+            "vswhere.exe",
+            "-latest",
+            "-find MSBuild/*/Bin/amd64/MSBuild.exe",
+        )
+        import subprocess as sp
+        proc = sp.Popen(" ".join(vswhere), stdout=sp.PIPE)
+        msbuild_exe = next(iter(proc.stdout.readline, b""), b"").decode().rstrip()
+        proc.stdout.close()
+
+        if not msbuild_exe:
+            msbuild_exe = "msbuild.exe"
+
+        engine_dir = self._engine.get_dir()
+        csproj_path = engine_dir / "Source/Programs/UnrealBuildTool/UnrealBuildTool.csproj"
+
+        args = (
+            str(csproj_path),
+            "/property:Configuration=Development;Platform=AnyCPU",
+            "/verbosity:m",
+            "/nologo",
+            "/target:Build",
+        )
+
+        yield msbuild_exe, args
+
+    def _get_script_paths(self):
+        return (
+            ("Build/BatchFiles/RunUBT.bat", False),
+            ("Build/BatchFiles/Build.bat", True),
+        )
+
+#-------------------------------------------------------------------------------
+class _BuildToolLinux(_BuildToolBase):
+    def _read_build_ubt_actions(self):
+        return iter(())
+
+    def _get_script_paths(self):
+        return (
+            ("Build/BatchFiles/RunUBT.sh", False),
+            ("Build/BatchFiles/Linux/Build.sh", True),
+        )
+
+#-------------------------------------------------------------------------------
+class _BuildToolMac(_BuildToolBase):
+    def _get_ubt_shim_script(self):
+        engine_dir = self._engine.get_dir()
+        script_path = engine_dir / "Intermediate/Build/ubt_ushell.sh"
+        if script_path.is_file():
+            return script_path
+
+        script_path.parent.mkdir(parents=True, exist_ok=True)
+        with script_path.open("wt") as out:
+            out.write("#!/bin/sh\n")
+            out.write(f"source {engine_dir}/Build/BatchFiles/Mac/SetupMono.sh\n")
+            out.write("xbuild \"$@\"\n")
+        script_path.chmod(0o777)
+        return script_path
+
+    def _read_build_ubt_actions(self):
+        try:    script_path = self._get_ubt_shim_script()
+        except: script_path = "xbuild"
+
+        engine_dir = self._engine.get_dir()
+        csproj_path = engine_dir / "Source/Programs/UnrealBuildTool/UnrealBuildTool.csproj"
+        yield str(script_path), (
+            csproj_path,
+            "/property:Configuration=Development",
+            "/verbosity:quiet",
+            "/nologo",
+            "/property:NoWarn=1591",
+        )
+
+    def _get_script_paths(self):
+        return (
+            ("Build/BatchFiles/RunUBT.sh", False),
+            ("Build/BatchFiles/Mac/Build.sh", True),
+        )
+
+#-------------------------------------------------------------------------------
+if sys.platform == "win32":     BuildTool = _BuildToolWin64
+elif sys.platform == "linux":   BuildTool = _BuildToolLinux
+elif sys.platform == "darwin":  BuildTool = _BuildToolMac
diff --git a/Engine/Extras/ushell/channels/unreal/core/pylib/unreal/cmdline.py b/Engine/Extras/ushell/channels/unreal/core/pylib/unreal/cmdline.py
new file mode 100644
--- /dev/null
+++ b/Engine/Extras/ushell/channels/unreal/core/pylib/unreal/cmdline.py
@@ -0,0 +1,62 @@
+# Copyright Epic Games, Inc. All Rights Reserved.
+
+import re
+from pathlib import Path
+from typing import Iterator
+
+# The Unreal Engine's parsing of command line arguments is unconventional. As
+# a user's given arguments flow through Python and C runtimes, their quoting is
+# canonicalised into a conventional-but-not-UE-compatible form. This module
+# attempts to format arguments in a UE-compatible way, wrapping them in a str-
+# derived class that deliberately fails '" " in s' and 'for x in s if x == " "'
+# style code such that the likes of subprocess.run do not add surrounding quotes
+
+_re = None
+
+#-------------------------------------------------------------------------------
+class _IterValue(str):
+    def __init__(self, inner):  self._inner = inner
+    def __eq__(self, rhs):      return False
+    def __str__(self):          return self._inner
+
+class _Iter(object):
+    def __init__(self, inner):  self._inner = inner
+    def __next__(self):         return _IterValue(next(self._inner))
+
+class _Shim(str):
+    def __init__(self, value):  self._value = value
+    def __iter__(self):         return _Iter(iter(str(self._value)))
+    def __contains__(self, k):  return False
+    def __str__(self):          return str(self._value)
+    def __fspath__(self):       return self
+
+#-------------------------------------------------------------------------------
+def _cache_re() -> re.Pattern:
+    return re.compile(r'^\s*-([a-zA-Z0-9]+)="*([^"]+)"*$')
+
+#-------------------------------------------------------------------------------
+def _ueify_impl(arg:str) -> str:
+    try: m = _re.match(arg)
+    except TypeError: return arg
+
+    if m and " " in m.group(2):
+        return _Shim(rf'-{m.group(1)}="{m.group(2)}"')
+    return arg
+
+#-------------------------------------------------------------------------------
+def read_ueified(*args:str) -> Iterator[str]:
+    global _re
+    _re = _re or _cache_re()
+    yield from (_ueify_impl(x) for x in args)
+
+#-------------------------------------------------------------------------------
+def ueified(*args:str) -> str:
+    return " ".join(read_ueified(*args))
+
+#-------------------------------------------------------------------------------
+def escape(arg):
+    if not isinstance(arg, _Shim):
+        return fr'"{arg}"' if " " in arg else arg
+    inner = str(arg._value)
+    inner = inner.replace('"', r'\"')
+    return _Shim(inner)
diff --git a/Engine/Extras/ushell/channels/unreal/core/pylib/unreal/platforms/linux/linux.py b/Engine/Extras/ushell/channels/unreal/core/pylib/unreal/platforms/linux/linux.py
new file mode 100644
--- /dev/null
+++ b/Engine/Extras/ushell/channels/unreal/core/pylib/unreal/platforms/linux/linux.py
@@ -0,0 +1,67 @@
+# Copyright Epic Games, Inc. All Rights Reserved.
+
+import os
+import unreal
+
+#-------------------------------------------------------------------------------
+class PlatformBase(unreal.Platform):
+    name = "Linux"
+    autosdk_name = "Linux_x64"
+    env_var = "LINUX_MULTIARCH_ROOT"
+
+    def _get_version_ue4(self):
+        dot_cs = self.get_unreal_context().get_engine().get_dir()
+        dot_cs /= "Source/Programs/UnrealBuildTool/Platform/Linux/UEBuildLinux.cs"
+        return Platform._get_version_helper_ue4(dot_cs, "ExpectedSDKVersion")
+
+    def _get_version_ue5(self):
+        dot_cs = "Source/Programs/UnrealBuildTool/Platform/Linux/LinuxPlatformSDK"
+        version = self._get_version_helper_ue5(dot_cs + ".Versions.cs")
+        return version or self._get_version_helper_ue5(dot_cs + ".cs")
+
+    def _get_cook_form(self, target):
+        if target == "game":   return "Linux" if self.get_unreal_context().get_engine().get_version_major() > 4 else "LinuxNoEditor"
+        if target == "client": return "LinuxClient"
+        if target == "server": return "LinuxServer"
+
+#-------------------------------------------------------------------------------
+class Platform(PlatformBase):
+    def _read_env(self):
+        env_var = PlatformBase.env_var
+        value = os.getenv(env_var)
+        if value:
+            yield env_var, value
+            return
+
+        version = self.get_version()
+        extras_dir = self.get_unreal_context().get_engine().get_dir() / "Extras"
+        for suffix in ("AutoSDK/", "ThirdPartyNotUE/SDKs/"):
+            value = extras_dir / suffix / f"HostLinux/Linux_x64/{version}/"
+            if (value / "x86_64-unknown-linux-gnu/bin").is_dir():
+                yield env_var, value
+                return
+
+        yield env_var, f"Linux_x64/{version}/"
+
+    def _launch(self, exec_context, stage_dir, binary_path, args):
+        if stage_dir:
+            base_dir = binary_path.replace("\\", "/").split("/")
+            base_dir = "/".join(base_dir[-4:-1])
+            base_dir = stage_dir + base_dir
+            if not os.path.isdir(base_dir):
+                raise EnvironmentError(f"Failed to find base directory '{base_dir}'")
+            args = (*args, "-basedir=" + base_dir)
+
+        ue_context = self.get_unreal_context()
+        engine_dir = ue_context.get_engine().get_dir()
+        engine_bin_dir = str(engine_dir / "Binaries/Linux")
+        if ld_lib_path := os.getenv("LD_LIBRARY_PATH"):
+            ld_lib_path += ":" + str(engine_bin_dir)
+        else:
+            ld_lib_path = str(engine_bin_dir)
+        env = exec_context.get_env()
+        env["LD_LIBRARY_PATH"] = ld_lib_path
+
+        cmd = exec_context.create_runnable(binary_path, *args)
+        cmd.launch()
+        return True
diff --git a/Engine/Extras/ushell/channels/unreal/core/pylib/unreal/platforms/mac/ios.py b/Engine/Extras/ushell/channels/unreal/core/pylib/unreal/platforms/mac/ios.py
new file mode 100644
--- /dev/null
+++ b/Engine/Extras/ushell/channels/unreal/core/pylib/unreal/platforms/mac/ios.py
@@ -0,0 +1,25 @@
+# Copyright Epic Games, Inc. All Rights Reserved.
+
+import os
+import unreal
+
+#-------------------------------------------------------------------------------
+class Platform(unreal.Platform):
+    name = "IOS"
+
+    def _read_env(self):
+        yield from ()
+
+    def _get_version_ue4(self):
+        dot_cs = self.get_unreal_context().get_engine().get_dir()
+        dot_cs /= "Source/Programs/UnrealBuildTool/Platform/Mac/MacToolChain.cs"
+        return Platform._get_version_helper_ue4(dot_cs, "MinMacOSVersion")
+
+    def _get_version_ue5(self):
+        dot_cs = "Source/Programs/UnrealBuildTool/Platform/Mac/ApplePlatformSDK"
+        version = self._get_version_helper_ue5(dot_cs + "Versions.cs")
+        return version or self._get_version_helper_ue5(dot_cs + ".cs")
+
+    def _get_cook_form(self, target):
+        if target == "game":   return "IOS"
+        if target == "client": return "IOSClient"
