Updated ushell [26875160]

diff --git a/Engine/Extras/ushell/channels/unreal/core/pylib/unreal/platforms/mac/mac.py b/Engine/Extras/ushell/channels/unreal/core/pylib/unreal/platforms/mac/mac.py
new file mode 100644
--- /dev/null
+++ b/Engine/Extras/ushell/channels/unreal/core/pylib/unreal/platforms/mac/mac.py
@@ -0,0 +1,41 @@
+# Copyright Epic Games, Inc. All Rights Reserved.
+
+import os
+import unreal
+
+#-------------------------------------------------------------------------------
+class Platform(unreal.Platform):
+    name = "Mac"
+
+    def _read_env(self):
+        yield from ()
+
+    def _get_version_ue4(self):
+        dot_cs = self.get_unreal_context().get_engine().get_dir()
+        dot_cs /= "Source/Programs/UnrealBuildTool/Platform/Mac/MacToolChain.cs"
+        return Platform._get_version_helper_ue4(dot_cs, "MinMacOSVersion")
+
+    def _get_version_ue5(self):
+        dot_cs = "Source/Programs/UnrealBuildTool/Platform/Mac/ApplePlatformSDK"
+        version = self._get_version_helper_ue5(dot_cs + ".Versions.cs")
+        return version or self._get_version_helper_ue5(dot_cs + ".cs")
+
+    def _get_cook_form(self, target):
+        if target == "game":   return "Mac" if self.get_unreal_context().get_engine().get_version_major() > 4 else "MacNoEditor"
+        if target == "client": return "MacClient"
+        if target == "server": return "MacServer"
+
+    def _launch(self, exec_context, stage_dir, binary_path, args):
+        if stage_dir:
+            bins_index = binary_path.find("/Binaries/Mac/")
+            if bins_index >= 0:
+                base_dir = stage_dir
+                base_dir += os.path.basename(binary_path[:bins_index])
+                base_dir += "/Binaries/Mac/"
+            else:
+                raise EnvironmentError(f"Unable to calculate base directory from '{binary_path}'")
+            args = (*args, "-basedir=" + base_dir)
+
+        cmd = exec_context.create_runnable(binary_path, *args)
+        cmd.launch()
+        return True
diff --git a/Engine/Extras/ushell/channels/unreal/core/pylib/unreal/platforms/mac/tvos.py b/Engine/Extras/ushell/channels/unreal/core/pylib/unreal/platforms/mac/tvos.py
new file mode 100644
--- /dev/null
+++ b/Engine/Extras/ushell/channels/unreal/core/pylib/unreal/platforms/mac/tvos.py
@@ -0,0 +1,25 @@
+# Copyright Epic Games, Inc. All Rights Reserved.
+
+import os
+import unreal
+
+#-------------------------------------------------------------------------------
+class Platform(unreal.Platform):
+    name = "TVOS"
+
+    def _read_env(self):
+        yield from ()
+
+    def _get_version_ue4(self):
+        dot_cs = self.get_unreal_context().get_engine().get_dir()
+        dot_cs /= "Source/Programs/UnrealBuildTool/Platform/Mac/MacToolChain.cs"
+        return Platform._get_version_helper_ue4(dot_cs, "MinMacOSVersion")
+
+    def _get_version_ue5(self):
+        dot_cs = "Source/Programs/UnrealBuildTool/Platform/Mac/ApplePlatformSDK"
+        version = self._get_version_helper_ue5(dot_cs + ".Versions.cs")
+        return version or self._get_version_helper_ue5(dot_cs + ".cs")
+
+    def _get_cook_form(self, target):
+        if target == "game":   return "TVOS"
+        if target == "client": return "TVOSClient"
diff --git a/Engine/Extras/ushell/channels/unreal/core/pylib/unreal/platforms/win64/android.py b/Engine/Extras/ushell/channels/unreal/core/pylib/unreal/platforms/win64/android.py
new file mode 100644
--- /dev/null
+++ b/Engine/Extras/ushell/channels/unreal/core/pylib/unreal/platforms/win64/android.py
@@ -0,0 +1,103 @@
+# Copyright Epic Games, Inc. All Rights Reserved.
+
+import os
+import re
+import shutil
+import unreal
+
+#-------------------------------------------------------------------------------
+class Platform(unreal.Platform):
+    name = "Android"
+
+    def _read_env(self):
+        version = self.get_version()
+        prefix = f"Android/{version}/"
+        env_vars = {}
+
+        # sdk
+        env_vars["ANDROID_HOME"] = "android-sdk-windows/" if version == "-22" else ""
+
+        # jdk
+        if version <= "-22":
+            env_vars["JAVA_HOME"] = "jdk1.8.0_77/"
+        elif version <= "-25":
+            env_vars["JAVA_HOME"] = "jre/"
+        else:
+            env_vars["JAVA_HOME"] = "jbr/"
+
+        # ndk
+        if version <= "-22":
+            env_vars["NDKROOT"] = "android-ndk-r14b/"
+        else:
+            if sdks_root := os.getenv("UE_SDKS_ROOT"):
+                from pathlib import Path
+                ndk_dir = Path(sdks_root)
+                ndk_dir /= "Host" + unreal.Platform.get_host()
+                ndk_dir /= prefix
+                ndk_dir /= "ndk"
+                if ndk_dir.is_dir():
+                    ndk_ver = max(x.name for x in ndk_dir.glob("*") if x.is_dir())
+                    env_vars["NDKROOT"] = "ndk/" + ndk_ver
+            env_vars.setdefault("NDKROOT", "ndk/")
+
+        # dispatch
+        for env_var, template in env_vars.items():
+            value = os.getenv(env_var)
+            if not value:
+                value = prefix + template
+            yield env_var, value
+
+    def _get_version_ue4(self):
+        dot_cs = self.get_unreal_context().get_engine().get_dir()
+        dot_cs /= "Source/Programs/UnrealBuildTool/Platform/Android/UEBuildAndroid.cs"
+        try:
+            import re
+            with open(dot_cs, "rt") as cs_in:
+                cs_in.read(8192)
+                lines = iter(cs_in)
+                next(x for x in lines if "override string GetRequiredSDKString()" in x)
+                next(lines) # {
+                for i in range(5):
+                    line = next(lines)
+                    if m := re.search(r'return "(-\d+)"', line):
+                        return m.group(1)
+        except (StopIteration, FileNotFoundError):
+            pass
+
+    def _get_version_ue5(self):
+        func = "GetAutoSDKDirectoryForMainVersion"
+        dot_cs = "Source/Programs/UnrealBuildTool/Platform/Android/AndroidPlatformSDK"
+        version = self._get_version_helper_ue5(dot_cs + ".Versions.cs", func)
+        return version or self._get_version_helper_ue5(dot_cs + ".cs", func)
+
+    def _get_android_home(self):
+        out = next((v for k,v in self._read_env() if k == "ANDROID_HOME"), None)
+        if not out:
+            return
+
+        for prefix in ("", unreal.Platform.get_sdks_dir()):
+            if os.path.isdir(prefix + out):
+                return prefix + out
+
+    def _get_adb(self):
+        home_dir = self._get_android_home()
+        return home_dir + "/platform-tools/adb" if home_dir else "adb"
+
+    def _get_aapt(self):
+        home_dir = self._get_android_home()
+        if not home_dir:
+            return "aapt"
+
+        try:
+            build_tools_dir = home_dir + "/build-tools/"
+            latest = ""
+            for item in os.scandir(build_tools_dir):
+                latest = max(latest, item.name)
+            if latest and os.path.isdir(build_tools_dir):
+                return build_tools_dir + latest + "/aapt"
+        except FileNotFoundError:
+            return
+
+    def _get_cook_form(self, target):
+        if target == "game":   return "Android_ASTC"
+        if target == "client": return "Android_ASTCClient"
diff --git a/Engine/Extras/ushell/channels/unreal/core/pylib/unreal/platforms/win64/linux.py b/Engine/Extras/ushell/channels/unreal/core/pylib/unreal/platforms/win64/linux.py
new file mode 100644
--- /dev/null
+++ b/Engine/Extras/ushell/channels/unreal/core/pylib/unreal/platforms/win64/linux.py
@@ -0,0 +1,14 @@
+# Copyright Epic Games, Inc. All Rights Reserved.
+
+import os
+from ..linux import linux
+
+#-------------------------------------------------------------------------------
+class Platform(linux.PlatformBase):
+    def _read_env(self):
+        env_var = linux.PlatformBase.env_var
+        value = os.getenv(env_var)
+        if not value:
+            version = self.get_version()
+            value = f"Linux_x64/{version}/"
+        yield env_var, value
diff --git a/Engine/Extras/ushell/channels/unreal/core/pylib/unreal/platforms/win64/win64.py b/Engine/Extras/ushell/channels/unreal/core/pylib/unreal/platforms/win64/win64.py
new file mode 100644
--- /dev/null
+++ b/Engine/Extras/ushell/channels/unreal/core/pylib/unreal/platforms/win64/win64.py
@@ -0,0 +1,147 @@
+# Copyright Epic Games, Inc. All Rights Reserved.
+
+import os
+import unreal
+import subprocess as sp
+
+#-------------------------------------------------------------------------------
+class _PlatformBase(unreal.Platform):
+    def _read_env(self):
+        yield from ()
+
+    def _get_version_ue4(self):
+        import re
+
+        dot_cs = self.get_unreal_context().get_engine().get_dir()
+        dot_cs /= "Source/Programs/UnrealBuildTool/Platform/Windows/UEBuildWindows.cs"
+
+        needles = (
+            ("DefaultWindowsSdkVersion", re.compile(r'DefaultWindowsSdkVersion\s+=\s+("|[\w.]+\("|)([^"]+)"?\)?\s*;')),
+            ("VersionNumber", re.compile(r'^\s+VersionNumber.Parse\(("([^"]+))')),
+        )
+        versions = []
+        try:
+            with open(dot_cs, "rt") as lines:
+                for line in lines:
+                    for seed, needle in needles:
+                        if seed not in line:
+                            continue
+                        m = needle.search(line)
+                        if m:
+                            versions.append(m.group(2))
+        except FileNotFoundError:
+            return
+
+        return self._get_version_string(versions)
+
+    def _get_version_ue5(self):
+        # Extract the Windows SDK version
+        dot_cs = self.get_unreal_context().get_engine().get_dir()
+        dot_cs /= "Source/Programs/UnrealBuildTool/Platform/Windows"
+        version = self._get_version_ue5_path(dot_cs / "MicrosoftPlatformSDK.Versions.cs")
+        return version or self._get_version_ue5_path(dot_cs / "MicrosoftPlatformSDK.cs")
+
+    def _get_version_ue5_path(self, dot_cs):
+        import re
+        needle = re.compile(r'VersionNumber.Parse\("([\d.]+)"')
+
+        versions = []
+        try:
+            with open(dot_cs, "rt") as lines:
+                for line in (x for x in lines if "VersionNumber" in x):
+                    if m := needle.search(line):
+                        versions.append(m.group(1))
+                        break
+        except FileNotFoundError:
+            pass
+
+        if not versions:
+            if sdk_version := self._get_version_helper_ue5(dot_cs):
+                versions.append(sdk_version)
+            else:
+                return
+
+        # Extract the Visual Studio toolchain version
+        vs_version = self._get_vs_toolchain_version_string(dot_cs)
+        if not vs_version:
+            vs_version = self._get_vs_toolchain_version_string(dot_cs.parent / "UEBuildWindows.cs")
+        if not vs_version:
+            return
+        versions.append(vs_version)
+
+        return self._get_version_string(versions)
+
+    def _get_vs_toolchain_version_string(self, dot_cs):
+        import re
+        try:
+            with open(dot_cs, "rt") as lines:
+                line_iter = iter(lines)
+                for line in line_iter:
+                    if "PreferredVisualCppVersions" in line:
+                        break
+                else:
+                    return
+
+                version = None
+                for line in line_iter:
+                    if "}" in line:
+                        break
+                    elif m := re.match(r'^\s+VersionNumberRange\.Parse\(.+, "([^"]+)', line):
+                        version = m.group(1)
+                        break
+
+                return version
+        except FileNotFoundError:
+            return
+
+    def _get_version_string(self, versions):
+        msvc_versions = list(x for x in versions if x >= "14.")
+        if not msvc_versions:
+            return
+
+        sdk_versions = list(x for x in versions if x.startswith("10."))
+        if not sdk_versions:
+            return
+
+        return max(sdk_versions) + "-" + max(msvc_versions)
+
+#-------------------------------------------------------------------------------
+class Platform(_PlatformBase):
+    name = "Win64"
+    config_name = "Windows"
+    autosdk_name = None
+    vs_transport = "Default"
+
+    def _get_cook_form(self, target):
+        if target == "game":   return "Windows" if self.get_unreal_context().get_engine().get_version_major() > 4 else "WindowsNoEditor"
+        if target == "client": return "WindowsClient"
+        if target == "server": return "WindowsServer"
+
+    def _launch(self, exec_context, stage_dir, binary_path, args):
+        if stage_dir:
+            base_dir = binary_path.replace("\\", "/").split("/")
+            base_dir = "/".join(base_dir[-4:-1])
+            base_dir = stage_dir + base_dir
+            if not os.path.isdir(base_dir):
+                raise EnvironmentError(f"Failed to find base directory '{base_dir}'")
+            args = (*args, "-basedir=" + base_dir)
+
+        cmd = exec_context.create_runnable(binary_path, *args)
+        cmd.launch()
+        return (cmd.get_pid(), None)
+
+    def _kill(self, target):
+        target_name = None
+        context = self.get_unreal_context()
+        if target == "editor":
+            ue_version = context.get_engine().get_version_major()
+            target_name = "UE4Editor" if ue_version == 4 else "UnrealEditor"
+        elif target:
+            target_type = unreal.TargetType.parse(target)
+            target_name = context.get_target_by_type(target_type).get_name()
+        if target_name:
+            print(f"Terminating {target_name}*.exe")
+            sp.run(rf'taskkill.exe /f /fi "imagename eq {target_name}*"')
+        else:
+            self._kill("client")
+            self._kill("game")
diff --git a/Engine/Extras/ushell/channels/unreal/core/pylib/unrealcmd.py b/Engine/Extras/ushell/channels/unreal/core/pylib/unrealcmd.py
new file mode 100644
--- /dev/null
+++ b/Engine/Extras/ushell/channels/unreal/core/pylib/unrealcmd.py
@@ -0,0 +1,115 @@
+# Copyright Epic Games, Inc. All Rights Reserved.
+
+import os
+import unreal
+import flow.cmd
+
+#-------------------------------------------------------------------------------
+def get_uproject_from_dir(start_dir):
+    from pathlib import Path
+    start_dir = Path(start_dir)
+
+    for search_dir in (start_dir, *start_dir.parents):
+        with os.scandir(search_dir) as files:
+            try: project = next(x.path for x in files if x.name.endswith(".uproject"))
+            except StopIteration: continue
+            return project
+
+    raise EnvironmentError(f"Unable to find a .uproject from '{start_dir}'")
+
+
+
+#-------------------------------------------------------------------------------
+Arg = flow.cmd.Arg
+Opt = flow.cmd.Opt
+class Cmd(flow.cmd.Cmd):
+    def complete_platform(self, prefix):
+        yield from self.read_platforms()
+
+    def complete_variant(self, prefix):
+        yield from (x.name.lower() for x in unreal.Variant)
+
+    def __init__(self):
+        super().__init__()
+        self._project = None
+        self._context = None
+
+    def get_unreal_context(self):
+        session = self.get_noticeboard(self.Noticeboard.SESSION)
+        path = session["uproject"] or "."
+        self._context = self._context or unreal.Context(path)
+        return self._context
+
+    def read_platforms(self):
+        platforms = self.get_unreal_context().get_platform_provider()
+        yield from platforms.read_platform_names()
+
+    def get_platform(self, name=None):
+        ue_context = self.get_unreal_context()
+        platforms = ue_context.get_platform_provider()
+        if platform := platforms.get_platform(name):
+            return platform
+
+        available = ",".join(platforms.read_platform_names())
+        raise ValueError(f"Unknown platform '{name}'. Valid list: {available}")
+
+    def get_host_platform(self):
+        return self.get_platform()
+
+
+
+#-------------------------------------------------------------------------------
+class MultiPlatformCmd(Cmd):
+    noplatforms = Opt(False, "Disable discovery of available platforms")
+
+    def __init__(self):
+        super().__init__()
+        self._platforms = []
+
+    def use_all_platforms(self):
+        self._platforms = list(self.read_platforms())
+
+    def use_platform(self, platform_name):
+        self._platforms.append(platform_name)
+
+    def get_exec_context(self):
+        context = super().get_exec_context()
+        if self.args.noplatforms or not self._platforms:
+            return context
+
+        env = context.get_env()
+        if len(self._platforms) == 1:
+            name = self._platforms[0]
+            platform = self.get_platform(name)
+            self.print_info(name.title(), platform.get_version())
+            for item in platform.read_env():
+                try:
+                    dir = str(item)
+                    env[item.key] = dir
+                except EnvironmentError:
+                    dir = flow.cmd.text.red(item.get())
+                print(item.key, "=", dir)
+            return context
+
+        self.print_info("Establishing platforms")
+
+        found = []
+        failed = []
+        for platform_name in self._platforms:
+            platform = self.get_platform(platform_name)
+            try:
+                platform_env = {x.key:str(x) for x in platform.read_env()}
+                env.update(platform_env)
+                found.append(platform_name)
+            except EnvironmentError:
+                failed.append(platform_name)
+
+        if found:
+            print("  found:", flow.cmd.text.green(" ".join(found)))
+        if failed:
+            out = "missing: "
+            out += flow.cmd.text.light_red(" ".join(failed))
+            out += flow.cmd.text.grey(" ('.info' for details)")
+            print(out)
+
+        return context
diff --git a/Engine/Extras/ushell/channels/unreal/core/pylib/vs/build.bat b/Engine/Extras/ushell/channels/unreal/core/pylib/vs/build.bat
new file mode 100755
--- /dev/null
+++ b/Engine/Extras/ushell/channels/unreal/core/pylib/vs/build.bat
@@ -0,0 +1,54 @@
+@echo off
+pushd %~dp0
+
+cl.exe 1>nul 2>nul
+if errorlevel 9009 (
+    for /f "delims=" %%d in ('vswhere -latest -find VC/Auxiliary/Build/vcvars64.bat') do (
+        call "%%d"
+    )
+)
+
+setlocal
+
+set vssdk=%~f1
+if "%vssdk%"=="" goto:usage
+
+set pysdk=%~f2
+if "%pysdk%"=="" goto:usage
+
+if not "%3"=="release" (
+    set _cl=/Od
+    set _link=/debug
+) else (
+    set _cl=/O1 /Os
+)
+
+set _cl="/I%vssdk%/Inc" "/I%pysdk%/include" /nologo /Zl /Zi /EHs-c- /c %_cl%
+set _link=%_link% /nologo /dll /machine:x64^
+    /nodefaultlib /align:16^
+    /incremental:no /opt:ref,icf^
+    "/libpath:%pysdk%/libs" python3.lib ole32.lib oleaut32.lib kernel32.lib
+
+mkdir 1>nul 2>nul _build
+cd _build
+
+call:go cl.exe ../dte.cpp %_cl%
+if errorlevel 1 exit /b 1
+
+call:go link.exe %_link% /out:dte.pyd dte.obj
+if errorlevel 1 exit /b 1
+
+move /y dte.pyd ..
+
+popd
+goto:eof
+
+:go
+echo.
+echo [96m%*[0m
+%*
+goto:eof
+
+:usage
+echo usage: %~nx0 vssdk pysdk [release]
+goto:eof
diff --git a/Engine/Extras/ushell/channels/unreal/core/pylib/vs/dte.cpp b/Engine/Extras/ushell/channels/unreal/core/pylib/vs/dte.cpp
new file mode 100644
--- /dev/null
+++ b/Engine/Extras/ushell/channels/unreal/core/pylib/vs/dte.cpp
@@ -0,0 +1,665 @@
+#define _HAS_EXCEPTIONS 0
+#define PY_SSIZE_T_CLEAN
+
+#include <objbase.h>
+#include <olectl.h>
+#include <Python.h>
+#include <new>
+
+#include <dte.h>
+#include <dte80.h>
+
+//------------------------------------------------------------------------------
+template <typename T>
+class ComPtr
+{
+public:
+                ComPtr()            = default;
+                ComPtr(void* t)     { _ptr = (T*)t; }
+                ~ComPtr()           { if (_ptr != nullptr) _ptr->Release(); }
+    explicit    operator bool ()    { return _ptr != nullptr; }
+                operator void** ()  { return (void**)&_ptr; }
+                operator T** ()     { return &_ptr; }
+                operator T* ()      { return _ptr; }
+    T*          operator -> ()      { return _ptr; }
+    T*          get()               { return _ptr; }
+    T*          detach()            { T* ret = _ptr; forget(); return ret; }
+    void        forget()            { _ptr = nullptr; }
+
+private:
+    T*          _ptr = nullptr;
+
+                ComPtr(const ComPtr&) = delete;
+                ComPtr(const ComPtr&&) = delete;
+    ComPtr&     operator = (const ComPtr&) = delete;
+    ComPtr&     operator = (const ComPtr&&) = delete;
+};
+
+//------------------------------------------------------------------------------
+class BStr
+{
+public:
+                BStr(const wchar_t* in) { _str = SysAllocString(in); }
+                ~BStr()                 { SysFreeString(_str); }
+                operator BSTR& ()       { return _str; }
+
+private:
+    BSTR        _str;
+                BStr() = delete;
+                BStr(const BStr&) = delete;
+                BStr(const BStr&&) = delete;
+    BStr&       operator = (const BStr&) = delete;
+    BStr&       operator = (const BStr&&) = delete;
+};
+
+
+
+//------------------------------------------------------------------------------
+class Dte
+    : public PyObject
+{
+public:
+                    Dte(_DTE* dte);
+    void            shutdown();
+    const wchar_t*  get_sln_path() const;
+    const wchar_t*  get_version() const;
+    bool            attach(int pid, const wchar_t* transport, const wchar_t* host) const;
+    bool            attach(const wchar_t* name, const wchar_t* transport, const wchar_t* host) const;
+    bool            open_file(const wchar_t* path, int line=-1) const;
+    bool            activate() const;
+
+private:
+    template <typename T>
+    bool            attach_impl(T name, const wchar_t* transport, const wchar_t* host) const;
+    _DTE*           _dte;
+};
+
+//------------------------------------------------------------------------------
+Dte::Dte(_DTE* dte)
+: _dte(dte)
+{
+    _dte->AddRef();
+}
+
+//------------------------------------------------------------------------------
+void Dte::shutdown()
+{
+    _dte->Release();
+}
+
+//------------------------------------------------------------------------------
+const wchar_t* Dte::get_sln_path() const
+{
+    Solution* outer;
+    if (FAILED(_dte->get_Solution(&outer)))
+        return nullptr;
+
+    ComPtr<_Solution> solution(outer);
+
+    BSTR out;
+    return SUCCEEDED(solution->get_FullName(&out)) ? out : nullptr;
+}
+
+//------------------------------------------------------------------------------
+const wchar_t* Dte::get_version() const
+{
+    BSTR out;
+    return SUCCEEDED(_dte->get_Version(&out)) ? out : nullptr;
+}
+
+//------------------------------------------------------------------------------
+bool Dte::attach(int pid, const wchar_t* transport, const wchar_t* host) const
+{
+    return attach_impl<int>(pid, transport, host);
+}
+
+//------------------------------------------------------------------------------
+bool Dte::attach(const wchar_t* name, const wchar_t* transport, const wchar_t* host) const
+{
+    return attach_impl<const wchar_t*>(name, transport, host);
+}
+
+//------------------------------------------------------------------------------
+template <typename T>
+bool Dte::attach_impl(T pid, const wchar_t* transport, const wchar_t* host) const
+{
+    transport = transport ? transport : L"Default";
+
+    // Get the debugger interface
+    ComPtr<Debugger> debugger;
+    if (FAILED(_dte->get_Debugger(debugger)))
+        return false;
+
+    ComPtr<Debugger2> debugger2;
+    if (FAILED(debugger->QueryInterface(__uuidof(Debugger2), debugger2)))
+        return false;
+
+    // Find the transport
+    ComPtr<Transports> transports;
+    debugger2->get_Transports(transports);
+
+    long n;
+    transports->get_Count(&n);
+    VARIANT index = { VT_I4 };
+    for (long i = 1; i <= n; ++i)
+    {
+        index.lVal = i;
+        ComPtr<Transport> candidate;
+        if (FAILED(transports->Item(index, candidate)))
+            continue;
+
+        BSTR candidate_name;
+        if (FAILED(candidate->get_Name(&candidate_name)))
+            continue;
+
+        if (lstrcmpiW(transport, candidate_name) != 0)
+            continue;
+
+        // Now find the process within the transport at attach to it.
+        ComPtr<Processes> processes;
+        if (FAILED(debugger2->GetProcesses(candidate, BStr(host), processes)))
+            return false;
+
+        struct
+        {
+            const wchar_t* preprocess(const wchar_t* input) const
+            {
+                int len = lstrlenW(input);
+                for (const wchar_t* c = input + len; c > input; --c)
+                    if (c[-1] == '\\' || c[-1] == '/')
+                        return c;
+
+                return input;
+            }
+
+            int preprocess(int input) const
+            {
+                return input;
+            }
+
+            bool operator () (Process* process, int pid) const
+            {
+                long candidate_pid;
+                if (FAILED(process->get_ProcessID(&candidate_pid)))
+                    return false;
+
+                return candidate_pid == pid;
+            }
+
+            bool operator () (Process* process, const wchar_t* name) const
+            {
+                BSTR candidate_name;
+                if (FAILED(process->get_Name(&candidate_name)))
+                    return false;
+
+                return lstrcmpiW(name, preprocess(candidate_name)) == 0;
+            }
+        } is_target_process;
+
+        pid = is_target_process.preprocess(pid);
+
+        long m;
+        processes->get_Count(&m);
+        VARIANT index = { VT_I4 };
+        for (long j = 1; j <= m; ++j)
+        {
+            index.lVal = j;
+            ComPtr<Process> process;
+            if (SUCCEEDED(processes->Item(index, process)))
+                if (is_target_process(process, pid))
+                    return SUCCEEDED(process->Attach());
+        }
+
+        return false;
+    }
+
+    return false;
+}
+
+//------------------------------------------------------------------------------
+bool Dte::open_file(const wchar_t* path, int line) const
+{
+    ComPtr<ItemOperations> item_ops;
+    if (FAILED(_dte->get_ItemOperations(item_ops)))
+        return false;
+
+    const wchar_t* view_kind = L"{7651A703-06E5-11D1-8EBD-00A0C90F26EA}"; // from vsViewKindTextView
+
+    ComPtr<Window> window;
+    if (FAILED(item_ops->OpenFile(BStr(path), BStr(view_kind), window)))
+        return false;
+
+    if (line < 0)
+        return true;
+
+    ComPtr<Document> document;
+    if (SUCCEEDED(_dte->get_ActiveDocument(document)))
+    {
+        ComPtr<IDispatch> selection;
+        if (SUCCEEDED(document->get_Selection(selection)))
+        {
+            ComPtr<TextSelection> text_selection;
+            if (SUCCEEDED(selection->QueryInterface(__uuidof(TextSelection), text_selection)))
+            {
+                text_selection->GotoLine(line, FALSE);
+            }
+        }
+    }
+
+    return true;
+}
+
+//------------------------------------------------------------------------------
+bool Dte::activate() const
+{
+    ComPtr<Window> main_window;
+    if (FAILED(_dte->get_MainWindow(main_window)))
+        return false;
+
+    return SUCCEEDED(main_window->Activate());
+}
+
+
+
+//------------------------------------------------------------------------------
+static PyTypeObject dte_type = {
+    PyVarObject_HEAD_INIT(nullptr, 0)
+};
+
+//------------------------------------------------------------------------------
+void dte_free(PyObject* object)
+{
+    auto* dte = (Dte*)object;
+    dte->shutdown();
+    PyMem_Free(object);
+}
+
+//------------------------------------------------------------------------------
+PyObject* dte_get_sln_path(PyObject* object, PyObject* args)
+{
+    auto* dte = (Dte*)object;
+    if (const wchar_t* name = dte->get_sln_path())
+        return PyUnicode_FromWideChar(name, -1);
+
+    Py_RETURN_NONE;
+}
+
+//------------------------------------------------------------------------------
+PyObject* dte_get_version(PyObject* object, PyObject* args)
+{
+    auto* dte = (Dte*)object;
+
+    if (const wchar_t* version = dte->get_version())
+        return PyUnicode_FromWideChar(version, -1);
+
+    Py_RETURN_NONE;
+}
+
+//------------------------------------------------------------------------------
+PyObject* dte_attach(PyObject* object, PyObject* args)
+{
+    auto* dte = (Dte*)object;
+
+    int argc = PyTuple_Size(args);
+
+    if (argc < 1)
+    {
+        PyErr_SetString(PyExc_TypeError, "Pid expected");
+        return nullptr;
+    }
+
+    // transport argument
+    wchar_t* transport = nullptr;
+    PyObject* py_transport = (argc > 1) ? PyTuple_GetItem(args, 1) : nullptr;
+    if (py_transport != nullptr && py_transport != Py_None)
+    {
+        if (!PyUnicode_Check(py_transport))
+            return PyErr_SetString(PyExc_TypeError, "Argument 1 must be a string (transport)"), nullptr;
+
+        transport = PyUnicode_AsWideCharString(py_transport, nullptr);
+    }
+
+    // host argument
+    wchar_t* host = nullptr;
+    PyObject* py_host = (argc > 2) ? PyTuple_GetItem(args, 2) : nullptr;
+    if (py_host != nullptr && py_host != Py_None)
+    {
+        if (!PyUnicode_Check(py_host))
+            return PyErr_SetString(PyExc_TypeError, "Argument 2 must be a string (host)"), nullptr;
+
+        host = PyUnicode_AsWideCharString(py_host, nullptr);
+    }
+
+    // process argument
+    bool ret = false;
+    PyObject* py_process = PyTuple_GetItem(args, 0);
+    if (PyLong_Check(py_process))
+    {
+        int pid = PyLong_AsLong(py_process);
+        ret = dte->attach(pid, transport, host);
+    }
+    else if (PyUnicode_Check(py_process))
+    {
+        wchar_t* name = PyUnicode_AsWideCharString(py_process, nullptr);
+        ret = dte->attach(name, transport, host);
+        PyMem_Free(name);
+    }
+    else
+        return PyErr_SetString(PyExc_TypeError, "Argument 0 must be an integer (pid) or string (name)"), nullptr;
+
+    PyMem_Free(host);
+    PyMem_Free(transport);
+
+    if (ret)
+        Py_RETURN_TRUE;
+
+    Py_RETURN_FALSE;
+}
+
+//------------------------------------------------------------------------------
+PyObject* dte_open_file(PyObject* object, PyObject* args)
+{
+    int argc = PyTuple_Size(args);
+
+    wchar_t* path = nullptr;
+    if (argc > 0)
+    {
+        PyObject* py_path = PyTuple_GetItem(args, 0);
+        if (!PyUnicode_Check(py_path))
+            return PyErr_SetString(PyExc_TypeError, "Argument 1 must be a string"), nullptr;
+
+        path = PyUnicode_AsWideCharString(py_path, nullptr);
+    }
+    else
+        Py_RETURN_FALSE;
+
+    int line = -1;
+    if (argc > 1)
+    {
+        PyObject* py_line = PyTuple_GetItem(args, 1);
+        if (!PyLong_Check(py_line))
+            return PyErr_SetString(PyExc_TypeError, "Argument 2 must be an integer"), nullptr;
+
+        line = PyLong_AsLong(py_line);
+    }
+
+    auto* dte = (Dte*)object;
+    bool ret = dte->open_file(path, line);
+
+    PyMem_Free(path);
+
+    if (ret)
+        Py_RETURN_TRUE;
+
+    Py_RETURN_FALSE;
+}
+
+//------------------------------------------------------------------------------
+PyObject* dte_activate(PyObject* object, PyObject* args)
+{
+    auto* dte = (Dte*)object;
+    if (dte->activate())
+        Py_RETURN_TRUE;
+
+    Py_RETURN_FALSE;
+}
+
+//------------------------------------------------------------------------------
+static PyMethodDef dte_type_methods[] = {
+    { "get_sln_path", dte_get_sln_path, METH_VARARGS, "Returns the path to the loaded .sln" },
+    { "get_version",  dte_get_version,  METH_VARARGS, "Returns VS' version" },
+    { "attach",       dte_attach,       METH_VARARGS, "Attaches to a process on a given transport" },
+    { "activate",     dte_activate,     METH_VARARGS, "Brings the VS instance to the foreground" },
+    { "open_file",    dte_open_file,    METH_VARARGS, "Opens a file, and optionally moves the cursor to a given line number" },
+    {},
+};
+
+
+
+//------------------------------------------------------------------------------
+class DteSet
+{
+public:
+    bool                    add(_DTE* dte);
+    _DTE*                   get(unsigned index) const   { return index < _count ? _dtes[index] : nullptr; }
+    unsigned                get_count() const           { return _count; }
+
+private:
+    static const unsigned   _MAX = 16;
+    unsigned                _count = 0;
+    _DTE*                   _dtes[_MAX];
+};
+
+//------------------------------------------------------------------------------
+bool DteSet::add(_DTE* dte)
+{
+    if (_count >= _MAX)
+        return false;
+
+    for (unsigned i = 0; i < get_count(); ++i)
+        if (_dtes[i] == dte)
+            return false;
+
+    _dtes[_count++] = dte;
+    return true;
+}
+
+//------------------------------------------------------------------------------
+class RunningIter
+    : public PyObject
+{
+public:
+                RunningIter();
+    void        shutdown();
+    _DTE*       next();
+
+private:
+    DteSet      _dtes;
+    unsigned    _iter = 0;
+};
+
+//------------------------------------------------------------------------------
+RunningIter::RunningIter()
+{
+    ComPtr<IRunningObjectTable> running_object_table;
+    GetRunningObjectTable(0, running_object_table);
+    if (!running_object_table)
+        return;
+
+    ComPtr<IEnumMoniker> enum_moniker;
+    running_object_table->EnumRunning(enum_moniker);
+    if (!enum_moniker)
+        return;
+
+    ComPtr<IBindCtx> bind_ctx;
+    CreateBindCtx(0, bind_ctx);
+    if (!bind_ctx)
+        return;
+
+    while (true)
+    {
+        ComPtr<IMoniker> moniker;
+        if (enum_moniker->Next(1, moniker, nullptr) == S_FALSE)
+            break;
+
+        ComPtr<IUnknown> unknown;
+        running_object_table->GetObject(moniker, unknown);
+
+        ComPtr<_DTE> dte;
+        unknown->QueryInterface(__uuidof(_DTE), dte);
+
+        if (dte && _dtes.add(dte))
+            dte.forget();
+    }
+}
+
+//------------------------------------------------------------------------------
+void RunningIter::shutdown()
+{
+    for (int i = 0, n = _dtes.get_count(); i < n; ++i)
+        _dtes.get(i)->Release();
+}
+
+//------------------------------------------------------------------------------
+_DTE* RunningIter::next()
+{
+    if (_iter >= _dtes.get_count())
+        return nullptr;
+
+    _DTE* dte = _dtes.get(_iter);
+    _iter++;
+
+    return dte;
+}
+
+
+
+//------------------------------------------------------------------------------
+static PyTypeObject running_iter_type = {
+    //PyVarObject_HEAD_INIT(nullptr, 0)
+};
+
+//------------------------------------------------------------------------------
+static PyObject* running_iter_new(PyTypeObject* type, PyObject* args, PyObject* kwargs)
+{
+    auto* iter = PyMem_New(RunningIter, sizeof(RunningIter));
+    new (iter) RunningIter();
+    return PyObject_Init(iter, &running_iter_type);
+};
+
+//------------------------------------------------------------------------------
+static void running_iter_free(PyObject* object)
+{
+    auto* iter = (RunningIter*)object;
+    iter->shutdown();
+    PyMem_Free(object);
+};
+
+//------------------------------------------------------------------------------
+static PyObject* running_iter_next(PyObject* object)
+{
+    auto* iter = (RunningIter*)object;
+    _DTE* dte = iter->next();
+    if (dte == nullptr)
+        return nullptr;
+
+    auto* py_dte = PyMem_New(Dte, sizeof(Dte));
+    new (py_dte) Dte(dte);
+    return PyObject_Init(py_dte, &dte_type);
+}
+
+
+
+//------------------------------------------------------------------------------
+static int (WINAPI *NtResumeProcess)(HANDLE);
+static int (WINAPI *NtSuspendProcess)(HANDLE);
+
+//------------------------------------------------------------------------------
+static void suspend_resume_init()
+{
+    HMODULE ntdll = GetModuleHandleW(L"ntdll");
+    NtResumeProcess = decltype(NtResumeProcess)(GetProcAddress(ntdll, "NtResumeProcess"));
+    NtSuspendProcess = decltype(NtSuspendProcess)(GetProcAddress(ntdll, "NtSuspendProcess"));
+}
+
+//------------------------------------------------------------------------------
+static PyObject* suspend_resume(bool is_suspend, PyObject** args, int nargs)
+{
+    if (nargs < 1)
+        Py_RETURN_FALSE;
+
+    PyObject* py_pid = args[0];
+    if (!PyLong_Check(py_pid))
+        Py_RETURN_FALSE;
+
+    unsigned pid = PyLong_AsLong(py_pid);
+
+    HANDLE handle = OpenProcess(PROCESS_SUSPEND_RESUME, FALSE, pid);
+    if (handle == nullptr)
+        Py_RETURN_FALSE;
+
+    int ok = is_suspend ? NtSuspendProcess(handle) : NtResumeProcess(handle);
+    CloseHandle(handle);
+
+    if (ok < 0)
+        Py_RETURN_FALSE;
+
+    Py_RETURN_TRUE;
+}
+
+//------------------------------------------------------------------------------
+static PyObject* suspend_process(PyObject* self, PyObject** args, Py_ssize_t nargs)
+{
+    return suspend_resume(true, args, int(nargs));
+}
+
+//------------------------------------------------------------------------------
+static PyObject* resume_process(PyObject* self, PyObject** args, Py_ssize_t nargs)
+{
+    return suspend_resume(false, args, int(nargs));
+}
+
+
+
+//------------------------------------------------------------------------------
+static PyMethodDef dte_module_methods[] = {
+    { "suspend_process", PyCFunction(suspend_process), METH_FASTCALL, "Suspend a process by PID"},
+    { "resume_process",  PyCFunction(resume_process),  METH_FASTCALL, "Resume a process by PID"},
+    {},
+};
+
+//------------------------------------------------------------------------------
+static PyModuleDef dte_module = {
+    PyModuleDef_HEAD_INIT,
+    "dte",
+    nullptr,
+    -1,
+    dte_module_methods,
+};
+
+//------------------------------------------------------------------------------
+extern "C" __declspec(dllexport)
+PyObject* PyInit_dte()
+{
+    CoInitialize(nullptr);
+    suspend_resume_init();
+
+    //&PyType_Type
+    running_iter_type.tp_name = "read_running";
+    running_iter_type.tp_doc = "Iterates running instances of Visual Studio";
+    running_iter_type.tp_basicsize = sizeof(RunningIter);
+    running_iter_type.tp_flags = Py_TPFLAGS_DEFAULT;
+    running_iter_type.tp_new = running_iter_new;
+    running_iter_type.tp_dealloc = running_iter_free;
+    running_iter_type.tp_iter = PyObject_SelfIter;
+    running_iter_type.tp_iternext = running_iter_next;
+    if (PyType_Ready(&running_iter_type) < 0)
+        return nullptr;
+
+    //&PyType_Type
+    dte_type.tp_name = "dte";
+    dte_type.tp_doc = "Class that wraps a _DTE COM object";
+    dte_type.tp_basicsize = sizeof(Dte);
+    dte_type.tp_flags = Py_TPFLAGS_DEFAULT;
+    dte_type.tp_dealloc = dte_free;
+    dte_type.tp_methods = dte_type_methods;
+    if (PyType_Ready(&dte_type) < 0)
+        return nullptr;
+
+    dte_module.m_free = [] (void*) { CoUninitialize(); };
+    PyObject* module = PyModule_Create(&dte_module);
+    PyModule_AddObject(module, "running", (PyObject*)&running_iter_type);
+    return module;
+}
+
+
+//------------------------------------------------------------------------------
+BOOL WINAPI
+_DllMainCRTStartup(
+        HANDLE  hDllHandle,
+        DWORD   dwReason,
+        LPVOID  lpreserved
+        )
+{
+    return TRUE;
+}
diff --git a/Engine/Extras/ushell/channels/unreal/core/tips.py b/Engine/Extras/ushell/channels/unreal/core/tips.py
new file mode 100644
--- /dev/null
+++ b/Engine/Extras/ushell/channels/unreal/core/tips.py
@@ -0,0 +1,15 @@
+# Copyright Epic Games, Inc. All Rights Reserved.
+
+#-------------------------------------------------------------------------------
+class Tips(object):
+    def get_tips(self):
+        return (
+            "Auto-attach a debugger with '.run ... --attach'",
+            "Stage a cook with the '.stage' command",
+            "Change project within a ushell session using '.project <uprojpath>'",
+            "To modify build configuration use '.build xml [set|clear|edit]'",
+            "Active build configuration will be shown when running '.build'",
+            "List current build configuration with '.build xml'",
+            "Generate a ClangDb (compile_commands.json) using '.build misc clangdb'",
+            "'.project' (no arguments) will fuzzy search for .uprojects from CWD",
+        )
diff --git a/Engine/Extras/ushell/channels/unreal/perforce/boot.py b/Engine/Extras/ushell/channels/unreal/perforce/boot.py
new file mode 100644
--- /dev/null
+++ b/Engine/Extras/ushell/channels/unreal/perforce/boot.py
@@ -0,0 +1,20 @@
+# Copyright Epic Games, Inc. All Rights Reserved.
+
+import os
+import p4utils
+import flow.cmd
+
+#-------------------------------------------------------------------------------
+class Boot(flow.cmd.Cmd):
+    def run(self, env):
+        self.print_info("Perforce environment")
+        env["P4IGNORE"] = ".p4ignore.txt"
+        env["P4CONFIG"] = p4utils.get_p4config_name()
+        if editor := p4utils.get_p4_set("P4EDITOR"):
+            env["P4EDITOR"] = editor
+
+        print("P4CONFIG = " + env["P4CONFIG"])
+        print("P4IGNORE = " + env["P4IGNORE"])
+        print()
+
+        return super().run(env)
diff --git a/Engine/Extras/ushell/channels/unreal/perforce/cmds/bisect.py b/Engine/Extras/ushell/channels/unreal/perforce/cmds/bisect.py
new file mode 100644
--- /dev/null
+++ b/Engine/Extras/ushell/channels/unreal/perforce/cmds/bisect.py
@@ -0,0 +1,291 @@
+# Copyright Epic Games, Inc. All Rights Reserved.
+
+import os
+import p4utils
+import unrealcmd
+from peafour import P4
+import subprocess as sp
+from pathlib import Path
+
+#-------------------------------------------------------------------------------
+class _Bisectomatron(object):
+    def __init__(self, length):
+        self._good = 0
+        self._bad = length - 1
+        self._jump = -1
+        self._ugly = set()
+        self._bisect()
+
+    def get_good(self):  return self._good
+    def get_bad(self):   return self._bad
+    def get_index(self): return self._index
+    def get_state(self): return self._good, self._index, self._bad, self._ugly
+
+    def __bool__(self):
+        return (self._bad - self._good) > len(self._ugly) + 1
+
+    def _bisect(self):
+        self._index = (self._good + self._bad) // 2
+        while self._index in self._ugly:
+            self._index += self._jump
+
+    def good(self):
+        self._ugly = {x for x in self._ugly if x > self._index}
+        self._good = self._index
+        self._jump = 1
+        self._bisect()
+
+    def bad(self):
+        self._ugly = {x for x in self._ugly if x < self._index}
+        self._bad = self._index
+        self._jump = -1
+        self._bisect()
+
+    def ugly(self):
+        self._ugly.add(self._index)
+        while self:
+            jump = self._jump
+            self._jump = -jump - (jump // abs(jump))
+            self._index += jump
+            if self._index not in self._ugly:
+                break
+
+
+
+#-------------------------------------------------------------------------------
+def _resolve_script_path(script_path):
+    import shutil
+    if whiched := shutil.which(script_path):
+        return Path(whiched)
+
+    script_path = Path(script_path).absolute()
+    if script_path.is_file():
+        return script_path
+
+    raise ValueError(f"Script '{script_path}' does not exist");
+
+
+
+#-------------------------------------------------------------------------------
+class Bisect(unrealcmd.Cmd):
+    """ Bisects a range of changelists between 'good' and 'bad' (exclusive). If
+    a failed-build is reported, then bisection expands outwards from the active
+    changelist until bisection resumed.
+
+    The optional 'script' argument can be used to automate a bisection. Which change
+    is considered next is controlled by the scripts exit code; 0=good, 80=bad, and
+    90=failed_build. The exit codes are chosen to avoid interference from general
+    scripting errors. Here is an example Batch script that could be used with the
+    command line '.p4 bisect 123 456 path/to/script.bat test_case debug';
+
+      @echo off
+      .build editor %2
+      if errorlevel 1 exit 90
+      .run editor %2 -- -stdout -ExecCmds="Automation RunTests %1;Quit"
+      if errorlevel 1 exit 80
+
+    Here the 'test_case' and 'debug' arguments given when running the bisect are
+    passed on to the script and appear in '%1' and '%2' respectively. """
+    good       = unrealcmd.Arg(int, "First changelist that is in good shape")
+    bad        = unrealcmd.Arg(int, "Changelist that is known to be bad")
+    script     = unrealcmd.Arg([str], "Path and arguments to a script to automate bisection")
+    dryrun     = unrealcmd.Opt(False, "Do not sync or run the given script")
+    silentsync = unrealcmd.Opt(False, "Hide the output from '.p4 sync'")
+    clsfromcwd = unrealcmd.Opt(False, "Use the current directory to build candidate changelists")
+
+    def _user_choice(self):
+        return input("[g]ood, [b]ad, [f]ailed-build? ").lower()
+
+    def _script_choice(self):
+        if self.args.dryrun:
+            return _user_choice()
+
+        ret = sp.run(self._script).returncode
+        if ret == 80: return "b"
+        if ret ==  0: return "g"
+        if ret == 90: return "f"
+
+        self.print_warning("Unexpected script exit code;", ret)
+        return self._user_choice()
+
+    def _get_candidate_changes(self, good, bad):
+        # Check we've a branch-type UE context to get paths from
+        ue_context = None
+        use_cwd = self.args.clsfromcwd
+        try:
+            ue_context = self.get_unreal_context()
+            use_cwd = use_cwd or (ue_context.get_branch() == None)
+        except EnvironmentError:
+            use_cwd = True
+
+        # Fetch changes
+        depot_paths = []
+        if use_cwd or not ue_context:
+            depot_paths += "..."
+        else:
+            depot_paths.append(ue_context.get_engine().get_dir() / "...")
+            if project := ue_context.get_project():
+                depot_paths.append(project.get_dir() / "...")
+
+        temp_cls = []
+        self.print_info("Fetching changes from Perforce")
+        for i, depot_path in enumerate(depot_paths):
+            print(i, depot_path, end="")
+            temp_cls += [x for x in P4.changes(f"{depot_path}@{good},@{bad}", L=True)]
+            print()
+
+        # Sort and dedupe
+        temp_cls.sort(key=lambda x: x.change)
+
+        changes = temp_cls[:1]
+        for i, candidate in enumerate(temp_cls[1:]):
+            if candidate.change != temp_cls[i].change:
+                changes.append(candidate)
+
+        # Can't bisect a range that's too small
+        if len(changes) < 3:
+            raise EnvironmentError(f"Perforce returned too few changelists between {good} and {bad}")
+        print(len(changes), " changes to bisect")
+
+        # Remove CRs and LFs for neat(ish) printing
+        for item in changes:
+            item.desc = item.desc.replace("\n", " ")
+            item.desc = item.desc.replace("\r", "")
+
+        return changes
+
+    def main(self):
+        p4utils.login()
+
+        # Work out how we're going to interrogate bisect progress
+        self.print_info("Operating mode")
+        if self.args.script:
+            script, *args = self.args.script
+            script = _resolve_script_path(script)
+            self._script = (str(script), *args)
+            choice = self._script_choice
+            print("Script:   ", script)
+            print("Arguments:", *args)
+        else:
+            print("Interactive")
+            choice = self._user_choice
+
+        # Our range of candidates submits must contain one good and one bad
+        # submit as we're really looking for a state change.
+        good, bad = self.args.good, self.args.bad
+        if good >= bad:
+            raise ValueError("Invalid range of changelists given (f >= l)")
+
+        # Ask Perforce for some changes
+        changes = self._get_candidate_changes(good, bad)
+
+        # Run the bisect
+        self.print_info("Bisection")
+        sync_stdout = sp.DEVNULL if self.args.silentsync else None
+        index_generator = _Bisectomatron(len(changes))
+        while index_generator:
+            index = index_generator.get_index()
+            change = changes[index]
+
+            # Sync
+            if not self.args.dryrun:
+                print(f"Syncing {change.change}...", end="" if sync_stdout else "\n")
+                sync_cl = changes[index].change
+                sync_ret = sp.run(
+                    ("_p4", "sync", sync_cl),
+                    stdout=sync_stdout
+                )
+                if exitcode := sync_ret.returncode:
+                    print()
+                    raise RuntimeError(f"Failed syncing changelist {sync_cl} (exitcode={exitcode})")
+                print(" done")
+
+            # Pause and and ask the user what to do next.
+            while True:
+                good = index_generator.get_good()
+                bad = index_generator.get_bad()
+                print(f"\n\n{changes[index].change} ({good},{index},{bad}) {changes[index].desc[:64]}")
+
+                c = choice()
+                if c == "g":
+                    # Search earlier
+                    index_generator.good()
+                    break
+                elif c == "b":
+                    # Search later
+                    index_generator.bad()
+                    break
+                elif c == "f":
+                    # Expand out from index looking for a successful compile
+                    index_generator.ugly()
+                    break
+                else:
+                    print(f"Unknown input '{c}'!")
+
+        good_index = index_generator.get_good()
+        bad_index = index_generator.get_bad()
+        good_cl = changes[good_index]
+        bad_cl = changes[bad_index]
+
+        print("")
+        print("")
+        print("Bisection complete")
+        print("")
+        print("Good change:", good_cl.change)
+        print("Bad change: ", bad_cl.change)
+
+
+
+#-------------------------------------------------------------------------------
+def _test_bisectomatron():
+    import random
+    random.seed(37)
+
+    j = b"0123456789abcdefghijklmnopqrstuvwxyzABCDEFGHIJKLMNOPQRSTUVWXYZ"
+    def run(length, error_at, ugly=None):
+        def print_legend():
+            if ugly:
+                ugs = bytearray(b" " * (ugly[1] + 1))
+                for x in range(ugly[0], ugly[1] + 1):
+                    ugs[x] = ord("x")
+                print(" ", ugs.decode(), sep="")
+            guide = ("_" * error_at) + "|" + ("-" * (length - error_at - 1))
+            print("[", guide, "]", sep="")
+
+        line = bytearray(b" " * length)
+        b = _Bisectomatron(length);
+        i = 0
+        while b:
+            index = b.get_index()
+            line[index] = j[i]
+            if ugly and index >= ugly[0] and index <= ugly[1]: b.ugly()
+            elif index < error_at: b.good()
+            else: b.bad()
+            i += 1
+        line[b.get_bad()] = ord("*")
+
+        ok = b.get_bad() == error_at
+        if ugly:
+            ok = error_at > b.get_good() and error_at <= b.get_bad()
+        if not ok:
+            print_legend()
+            print(" ", line.decode(), sep="")
+            print()
+            print(b.get_good(), b.get_bad())
+            breakpoint()
+
+    for i in range(50000):
+        print(i, end="\r")
+        length = random.randint(0, 55) + 3
+        error_at = 1 + random.randint(0, length - 2)
+        run(length, error_at)
+    print()
+
+    for i in range(50000):
+        print(i, end="\r")
+        length = random.randint(3, 55)
+        error_at = random.randint(1, length - 2)
+        ugly = [0, 0]
+        ugly[0] = random.randint(1, length - 2)
+        ugly[1] = random.randint(ugly[0], length - 2)
+        run(length, error_at, ugly)
diff --git a/Engine/Extras/ushell/channels/unreal/perforce/cmds/cherrypick.py b/Engine/Extras/ushell/channels/unreal/perforce/cmds/cherrypick.py
new file mode 100644
--- /dev/null
+++ b/Engine/Extras/ushell/channels/unreal/perforce/cmds/cherrypick.py
@@ -0,0 +1,503 @@
+# Copyright Epic Games, Inc. All Rights Reserved.
+
+import os
+import re
+import p4utils
+import flow.cmd
+import subprocess
+from peafour import P4
+
+#-------------------------------------------------------------------------------
+def _untag_desc(desc):
+    tags = ("#fyi", "#review", "#codereview", "#robomerge")
+    out = ""
+    for line in desc.splitlines():
+        for tag in (x for x in tags if x in line):
+            line = line.replace(tag, f"#[{tag[1:]}]")
+        out += line + "\n"
+
+    return out
+
+
+
+#-------------------------------------------------------------------------------
+class Cherrypick(flow.cmd.Cmd):
+    """ Integrates/unshelves one or more changelists into the current branch,
+    resolving and clearing integration records as required. To clear integration
+    records only give a single pending changelist from the current client. Providing
+    a value for 'path' allows the cherrypick to be limited to a subset of files.
+    'path' should be either relative to a branch root; Engine/.../Runtime, or a
+    full depot path at the destination; //Dest/Stream/Engine/Source/... t.ex. The
+    command will *not* submit anything on your behalf. If in doubt you can use
+    '--dryrun' to preview the action. """
+    changelist    = flow.cmd.Arg([int], "Changelist or shelve to edit-integrate")
+    path          = flow.cmd.Opt("", "Restrict cherrypick to the given path (e.g. Engine/...)")
+    saferesolve   = flow.cmd.Opt(False, "Resolve safely and not automatically")
+    dryrun        = flow.cmd.Opt(False, "Only pretend to do the cherrypick")
+    force         = flow.cmd.Opt(False, "Force the operation through")
+    alwayseddy    = flow.cmd.Opt(False, "Always edigrate the result regardless of relation")
+    noeddy        = flow.cmd.Opt(False, "Skip the edigrate step")
+    sync          = flow.cmd.Opt(False, "Syncs target files to head before resolving")
+    virtual       = flow.cmd.Opt(False, "Perform the integration server-side")
+    rawbranchspec = flow.cmd.Opt(False, "Ignore the internal branchspec that Perforce generates")
+
+    complete_changelist = False
+
+    def get_explicit_relations(self):
+        return ()
+
+    def main(self):
+        if not self.args.changelist:
+            raise ValueError("No changelist(s) given")
+
+        # Check the user's logged into the current Perforce environment
+        username = p4utils.login()
+
+        # Get info about where we're cherrypicking to
+        self.print_info("Determining destination")
+        info = P4.info().run()
+        print("Server:", getattr(info, "proxyAddress", info.serverAddress))
+        print("Client:", info.clientName)
+
+        if info.clientName == "*unknown*":
+            raise EnvironmentError("Unknown Perforce client")
+
+        # Don't let the user cherrypick to somewhere out of where they think they
+        # are inadvertently.
+        cwd = os.getcwd()
+        if os.path.normpath(info.clientRoot).lower() not in cwd.lower():
+            raise EnvironmentError(f"Directory '{cwd}' is not under client '{info.clientName}'")
+
+        dest_where = P4.where("Engine").depotFile
+        dest_root = p4utils.get_branch_root(dest_where)
+        print("Root path: " + dest_root)
+
+        # Condition self.args.path
+        if path := self.args.path:
+            path = path.replace("\\", "/")
+            if path.startswith("//"):
+                path = "//" + "/".join(x for x in path.split("/") if x)
+                if not path.lower().startswith(dest_root.lower()):
+                    raise ValueError(f"'{path}' is unrelated to '{dest_root}'")
+            else:
+                path = dest_root + "/".join(x for x in path.split("/") if x)
+            self.args.path = path
+
+        # Get information about each input changelist
+        self.print_info("Fetching info on changelists to process")
+        descs = []
+        for cl in sorted(self.args.changelist):
+            print(cl, end="")
+            desc = P4.describe(str(cl), S=True).run()
+            descs.append(desc)
+            print(":", desc.desc.replace("\n", " ")[:68])
+
+        # If there's only one changelist and it's pending on the current client
+        # then clear integration records.
+        if len(descs) == 1:
+            desc = descs[0]
+            if desc.client == info.clientName and desc.status == "pending":
+                print("Change is pending changelist on current client")
+                return self._clear_integration_records(desc.change)
+
+        # Work out the branch root of each input changelist
+        self.print_info("Determining branch roots")
+        pick_rota = []
+        branch_roots = set()
+        for desc in descs:
+            path = getattr(desc, "path", None)
+            if not path:
+                depot_file = getattr(desc, "depotFile", None)
+                path = next(depot_file) if depot_file else ""
+
+            for branch_root in branch_roots:
+                if path.startswith(branch_root):
+                    break
+            else:
+                branch_root = p4utils.get_branch_root(path) if path else ""
+                branch_roots.add(branch_root)
+
+            print(desc.change, branch_root)
+            pick_rota.append((branch_root, desc))
+
+        # Validate that we can proceed.
+        self.print_info("Validating input changelists")
+        for root, desc in pick_rota:
+            print(desc.change, end="")
+            if desc.status == "pending":
+                if not hasattr(desc, "shelved"):
+                    print(" ... fail")
+                    raise ValueError(f"Pending changelist {desc.change} has no shelved files")
+            elif root == dest_root:
+                print(" ... fail")
+                raise ValueError(f"Changelist {desc.change} already submitted to "
+                 f"cherrypick destination '{dest_root}'")
+            elif not root:
+                print(" ... fail")
+                raise ValueError(f"Unable to determine branch root for {desc.change}")
+
+            print(" ... ok")
+
+        # Get a set of the destinations files
+        dest_paths = set()
+        for root, desc in pick_rota:
+            skip = len(root)
+            for src_path in desc.depotFile:
+                dest_paths.add(dest_root + src_path[skip:])
+
+        # Filter destination files by --path=
+        if path_spec := self.args.path:
+            print(f"Limiting input to {path_spec} : ", end="")
+
+            path_spec = path_spec.replace("...", "@")
+            path_spec = path_spec.replace(".", "\\.")
+            path_spec = path_spec.replace("*", "[^/]*")
+            path_spec = path_spec.replace("@", ".*")
+
+            prev_dest_paths_len = len(dest_paths)
+            dest_paths = [x for x in dest_paths if re.search(path_spec, x, re.IGNORECASE)]
+            print(prev_dest_paths_len - len(dest_paths), "of", prev_dest_paths_len, "removed")
+
+        # Check that none of the destination files are aleady opened for edit
+        self.print_info("Validating destination")
+        print("Affected files:", len(dest_paths))
+        print("Checking for open files")
+
+        validate = not self.args.force
+        if len(dest_paths) > 500:
+            validate = False
+            self.print_warning("Skipping validation. Source changelist is too big")
+
+        if validate:
+            def read_dest_paths_and_count():
+                count = 1
+                for x in dest_paths:
+                    print("\r" + str(count), "file(s) checked", end="")
+                    count += 1
+                    yield x
+
+            blocked = False
+            opened = P4.opened(read_dest_paths_and_count())
+            for item in opened:
+                self.print_warning("\r", item.depotFile[len(dest_root):])
+                blocked = True
+
+            if blocked:
+                raise EnvironmentError("The cherrypick of changelist"
+                    f" {desc.change} is blocked because some files are"
+                    " already open for edit at the destination.")
+
+            print()
+            print("All good. No target files were found to be already open for edit")
+        elif self.args.force:
+            self.print_warning("Validation was explicitly skipped")
+
+        # Create a target changelist for the cherrypick(s)
+        cl_desc = ""
+        for _, desc in pick_rota:
+            cl_desc += _untag_desc(desc.desc).rstrip()
+            cl_desc += "\n"
+
+        for _, desc in pick_rota:
+            cl_desc += f"\n#ushell-cherrypick of {desc.change} by {desc.user}"
+
+        cl_spec = { "Change" : "new", "Description" : cl_desc }
+        P4.change(i=True).run(input_data=cl_spec)
+        dest_cl = P4.changes(c=info.clientName, m=1, s="pending").change
+
+        # Integrate or unshelve each input changelist
+        specs = {}
+        self.print_info("Cherrypicking into", dest_cl)
+        for src_root, desc in pick_rota:
+            cl = desc.change
+
+            branch_spec, file_spec = specs.get(src_root, (None, None))
+            if not branch_spec:
+                print("Creating branch spec for", src_root)
+                branch_spec = p4utils.TempBranchSpec("cherrypick", username, src_root, dest_root, self.args.rawbranchspec)
+                print(str(branch_spec))
+
+                file_spec = src_root + "..."
+                if path := self.args.path:
+                    file_spec = src_root + path[len(dest_root):]
+                    print(f"Limiting to path '{file_spec}'")
+
+                specs[src_root] = branch_spec, file_spec
+
+            def on_info(info):
+                output = info.data
+                if "must resolve" in output or "also opened" in output:
+                    return
+                print("")
+                self.print_warning(info.data)
+
+            p4_args = {
+                "b" : branch_spec,
+                "n" : self.args.dryrun,
+                "f" : self.args.force,
+                "v" : self.args.virtual,
+                "c" : dest_cl,
+            }
+            if desc.status == "pending":
+                p4_args["-bypass-exclusive-lock"] = True
+                print(f"Unshelving {cl} from", src_root, end="")
+                file_spec_dest = file_spec.replace(src_root, dest_root)
+                unshelve = P4.unshelve(file_spec_dest, s=cl, **p4_args)
+                for x in unshelve.read(on_info=on_info): pass
+                print("")
+            else:
+                print(f"Integrating {cl} from", src_root, end="")
+                integrate = P4.integrate(s=file_spec + f"@{cl},@{cl}", **p4_args)
+                for x in integrate.read(on_info=on_info): pass
+                print("")
+
+        # We can't do anymore pretending after this point.
+        if self.args.dryrun:
+            P4.change(dest_cl, d=True).run()
+            return
+
+        # Collect files that might be in other changelists
+        if self.args.force:
+            P4.reopen(dest_paths, c=dest_cl).run(on_error=False)
+
+        self.print_info("Resolving")
+
+        # Sync to latest if requested to do so
+        if self.args.sync:
+            print("Syncing to head first (--sync)", end="")
+            def read_sync_paths(cl):
+                for item in P4.opened(c=cl):
+                    yield item.depotFile
+            sync = P4.sync(read_sync_paths(dest_cl), q=True)
+            for x in sync.read(on_error=False):
+                pass
+            print("")
+
+        # Resolve
+        resolve_count = 0
+
+        def print_error(error):
+            print("\r", end="")
+            msg = error.data.rstrip()
+            if msg.startswith("No file(s)"):
+                print(msg)
+            else:
+                self.print_error(msg)
+
+        resolve_args = {
+            "as" : self.args.saferesolve,
+            "am" : not self.args.saferesolve,
+            "f"  : self.args.force,
+            "c"  : dest_cl,
+        }
+        resolve = P4.resolve(**resolve_args)
+        for item in resolve.read(on_error=print_error):
+            if getattr(item, "clientFile", None):
+                resolve_count += 1
+                print("\r" + str(resolve_count), "file(s) resolved", end="")
+        if resolve_count:
+            print("")
+
+        # Report conflicted files
+        resolve = P4.resolve(c=dest_cl, n=True)
+        for item in resolve.read(on_error=False):
+            if name := getattr(item, "fromFile", None):
+                self.print_error(name[len(src_root):])
+
+        # Queue up some details to display when everything's complete
+        class OnExit(object):
+            def __del__(self):
+                print("\nCherrypick complete;", dest_cl)
+        on_exit = OnExit()
+
+        # If there are branch/integrates in dest_cl there's nothing more to do
+        for item in P4.opened(c=dest_cl):
+            if item.action in ["integrate", "branch"]:
+                break
+        else:
+            print("No integrated files found")
+            return
+
+        # Post-process the cherrypick'd files.
+        self.print_info("Processing integration records")
+        if self.args.noeddy:
+            self.print_warning("Skipped")
+            return
+
+        # There's nothing more to do if the source and destination are related.
+        relations_allowed = (len(specs) == 1)
+        relations_allowed &= (descs[0].status != "pending")
+        relations_allowed &= not self.args.alwayseddy
+        dest_stream = getattr(info, "clientStream", None)
+        if dest_stream and relations_allowed:
+            print("Checking relations between src and dest")
+
+            def get_stream_parent(stream):
+                stream_info = P4.stream(stream, o=True).run()
+                parent = stream_info.Parent
+                while parent.startswith("//"):
+                    stream_info = P4.stream(parent, o=True).run()
+                    if stream_info.Type != "virtual":
+                        break
+                    parent = stream_info.Parent
+                return parent
+
+            # Check if we've fetched from a parent stream.
+            dest_parent = get_stream_parent(dest_stream)
+            if src_root.startswith(dest_parent):
+                print(f"Not required; {src_root} and {dest_root} are related")
+                return
+
+            # Check if we've fetch from a child stream
+            src_client = P4.client(descs[0].client, o=True).run()
+            src_stream = getattr(src_client, "Stream", None)
+            if src_stream and dest_root.startswith(get_stream_parent(src_stream)):
+                print(f"Not required; {src_root} and {dest_root} are related")
+                return
+
+        # No stream relation but maybe there's an explicit branch relation
+        if relations_allowed:
+            explicit_relations = self.get_explicit_relations()
+            for relatives in explicit_relations:
+                related = src_root.startswith(relatives[0]) and dest_root.startswith(relatives[1])
+                related |= src_root.startswith(relatives[1]) and dest_root.startswith(relatives[0])
+                if related:
+                    print("Not required; explicit relations;")
+                    print("             ", relatives[0])
+                    print("             ", relatives[1])
+                    return
+
+        ret = self._clear_integration_records(dest_cl)
+
+        return ret
+
+    def _resolve_prompt(self, changelist):
+        P4.resolve(c=changelist, n=True).run()
+        self.print_error("There are pending conflicts which must be resolved before continuing")
+        print("Now you have a few courses of action available to you;")
+        print("  [r]esolve with P4V")
+        print("  [c]ommand line resolve")
+        print("  re[v]ert and exit")
+        print("  retry [Enter]")
+        print("  abort [Ctrl-C]")
+
+        choice = input("Which one do you fancy? [r/c/v/Enter/Ctrl-C] ")
+        try:
+            if choice == "c":
+                cmd = ("p4", "resolve", "-du", "-c", changelist)
+                subprocess.run(cmd)
+            elif choice == "r":
+                #p4utils.run_p4vc("p4vc", "resolve", "-c", changelist, "...") # this just doesn't work...
+                p4utils.run_p4vc("pendingchanges")
+                input("Press Enter to continue (Ctrl-C to abort)...")
+            elif choice == "v":
+                self.print_info("")
+                cmd = ("p4", "revert", "-wc", changelist, "...")
+                subprocess.run(cmd)
+                cmd = ("p4", "change", "-d", changelist)
+                subprocess.run(cmd)
+                return False
+        except FileNotFoundError:
+            self.print_error("Failed to run command;", *cmd)
+
+        print()
+
+    def _clear_integration_records(self, changelist):
+        self.print_info("Clearing integration records from", changelist)
+
+        # Editgrates are destructive so we'll need to interact with the user if
+        # they have pending resolves.
+        try:
+            while True:
+                self._resolve_prompt(changelist)
+        except P4.Error:
+            pass
+
+        # Check the pending changelist is a pending
+        open_files = list(P4.opened(c=changelist))
+        if not open_files:
+            self.print_warning(f"Changelist {changelist} is empty")
+            return False
+
+        # Work out what to do with each file
+        to_edit = []
+        to_add = []
+        to_delete = []
+        to_move = []
+        action_map = {
+            "edit"          : to_edit,
+            "integrate"     : to_edit,
+            "add"           : to_add,
+            "branch"        : to_add,
+            "delete"        : to_delete,
+            "move/delete"   : to_move,
+            "move/add"      : None,
+        }
+
+        for item in open_files:
+            action_list = action_map.get(item.action)
+            if action_list != None:
+                action_list.append(item)
+
+        print("   Adds:", len(to_add))
+        print("  Edits:", len(to_edit))
+        print("Deletes:", len(to_delete))
+        print("  Moves:", len(to_move))
+        print("  Total:", len(open_files))
+
+        self.print_info("Clearing integration records")
+
+        # Sync any edited files that are not on the client
+        to_sync = [x.depotFile for x in to_edit if x.haveRev == "none"]
+        if to_sync:
+            print("Syncing edited files that are not on the client...", end="")
+            for item in P4.sync(to_sync, q=True):
+                pass
+            print("done")
+
+        def read_files(items):
+            for item in items:
+                yield item.depotFile
+
+        # Server-revert and edit/add/delete files to drop integrate records.
+        print("Reverting server side .. ", end="")
+        revert = P4.revert("//...", c=changelist, k=True)
+        for action in revert:
+            pass
+        print("done")
+
+        p4args = {
+            "c" : changelist,
+        }
+
+        if to_edit:
+            print("Reopening .............. ", end="")
+            for item in P4.edit(read_files(to_edit), **p4args):
+                pass
+            print("done")
+
+        if to_add:
+            print("Adding ................. ", end="")
+            for item in P4.add(read_files(to_add), c=changelist):
+                pass
+            print("done")
+
+        if to_delete:
+            print("Deleting ............... ", end="")
+            for item in P4.delete(read_files(to_delete), **p4args):
+                pass
+            print("done")
+
+        if to_move:
+            print("Applying moves;")
+            print("  Reopening ............ ", end="")
+            for item in P4.edit(read_files(to_move), **p4args):
+                pass
+            print("done")
+
+            print("  Moving ............... ", end="")
+            for item in to_move:
+                P4.move(item.depotFile, item.movedFile, **p4args).run()
+            print("done")
+
+        return True
diff --git a/Engine/Extras/ushell/channels/unreal/perforce/cmds/clean.py b/Engine/Extras/ushell/channels/unreal/perforce/cmds/clean.py
new file mode 100644
--- /dev/null
+++ b/Engine/Extras/ushell/channels/unreal/perforce/cmds/clean.py
@@ -0,0 +1,291 @@
+# Copyright Epic Games, Inc. All Rights Reserved.
+
+import os
+import unreal
+import flow.cmd
+from peafour import P4
+import subprocess as sp
+from pathlib import Path
+
+#-------------------------------------------------------------------------------
+class _Collector(object):
+    def __init__(self):
+        self._dirs = []
+        self._careful_dirs = []
+        self._pinned_files = set()
+
+    def get_pinned_count(self):
+        return len(self._pinned_files)
+
+    def add_dir(self, path, *, remove_carefully=False):
+        (self._careful_dirs if remove_carefully else self._dirs).append(path)
+
+    def pin_file(self, path):
+        path = str(path.resolve())
+        self._pinned_files.add(path.lower())
+
+    def read_careful_dirs(self):
+        return (x for x in self._careful_dirs)
+
+    def dispatch(self, actuator):
+        actuator.begin()
+
+        for dir in self._dirs:
+            actuator.remove_dir(dir)
+
+        for dir in self._careful_dirs:
+            for item in (x for x in dir.rglob("*") if x.is_file()):
+                key = str(item.resolve()).lower()
+                if key not in self._pinned_files:
+                    actuator.remove_file(item)
+
+        actuator.end()
+
+#-------------------------------------------------------------------------------
+class _DryRun(object):
+    def _spam(self):
+        mb_size = format(self._bytes // 1024 // 1024, ",") + "MB"
+        print("%d files to remove (%s)" % (self._count, mb_size), end="\r")
+
+    def begin(self):
+        self._bytes = 0
+        self._count = 0
+        self._breakdown = {}
+
+    def remove_dir(self, path):
+        for item in (x for x in path.rglob("*") if x.is_file()):
+            self.remove_file(item)
+
+    def remove_file(self, path):
+        self._bytes += path.stat().st_size
+        self._count += 1
+        if self._count % 493:
+            self._spam()
+
+    def end(self):
+        self._spam()
+        print()
+
+#-------------------------------------------------------------------------------
+class _Cleaner(object):
+    def __init__(self, work_dir):
+        self._work_dir = work_dir
+
+    def __del__(self):
+        if self._dir_count:
+            print(f"Launching background rmdir for {self._rubbish_dir}")
+        else:
+            print("Skipping background rmdir. No major directories to remove")
+
+        if os.name == "nt":
+            sp.Popen(
+                ("cmd.exe", "/c", "rd", "/q/s", str(self._rubbish_dir)),
+                stdout=sp.DEVNULL,
+                stderr=sp.DEVNULL
+                )
+            pass
+        else:
+            if self._rubbish_dir.is_dir() and self._rubbish_dir != "/": # paranoia!
+                sp.Popen(
+                    ("rm", "-rf", str(self._rubbish_dir)),
+                    stdout=sp.DEVNULL,
+                    stderr=sp.DEVNULL)
+
+    def _spam(self):
+        print(f"Moved {self._dir_count} directories, {self._file_count} files removed", end="\r", flush=True)
+
+    def begin(self):
+        rubbish_dir = self._work_dir / ".ushell_clean"
+        rubbish_dir.mkdir(parents=True, exist_ok=True)
+        self._rubbish_dir = rubbish_dir
+
+        self._dir_count = 0
+        self._file_count = 0
+        print(f"Moving directories to {self._rubbish_dir.name} and removing unversioned files")
+
+    def remove_dir(self, path):
+        dest_name = "%08x_%016x_%s_%s" % (os.getpid(), id(path), path.parent.name, path.name)
+        try:
+            path.rename(str(self._rubbish_dir / dest_name))
+            self._dir_count += 1
+            self._spam()
+        except OSError as e:
+            print("WARNING:", e)
+
+    def remove_file(self, item):
+        item.chmod(0o666)
+        item.unlink()
+        self._file_count += 1
+        self._spam()
+
+    def end(self):
+        if self._file_count or self._dir_count:
+            print()
+
+
+
+#-------------------------------------------------------------------------------
+class Clean(flow.cmd.Cmd):
+    """ Cleans intermediate and temporary files from an Unreal Engine branch. The
+    following sub-directories of .uproject, .uplugin, and Engine/ are cleaned up;
+
+      Intermediate     - Removed entirely
+      DerivedDataCache - Removed entirely
+      Binaries         - Unversioned files are removed
+      Saved            - All sub-directories except --savedkeeps=... (see below)
+
+    Only a subset of sub-directories of Saved/ are removed. Any directories that
+    match --savedkeeps's comma-separated list are not removed. For example, to
+    clean everything excepted Saved/StagedBuilds/ and Saved/Profiling/;
+
+      .p4 clean --savedkeeps=StagedBuilds,Profiling
+
+    "StagedBuilds,Profiling" is the default value for --savedkeeps. If --allsaved
+    is given then all of Saved/ will be removed.
+
+    Note that the removal happens in two stages. First directories are moved into
+    .ushell_clean/ in the root of the branch. This directory is then removed in
+    the background after `.p4 clean` exits."""
+    dryrun     = flow.cmd.Opt(False, "Do nothing except reports statistics")
+    allsaved   = flow.cmd.Opt(False, "Completely clean Saved/ directories")
+    savedkeeps = flow.cmd.Opt("Profiling,StagedBuilds", "Comma-separated list of Saved/ sub-directories to keep")
+
+    def _append_saved(self, collector, saved_dir):
+        if self.args.allsaved:
+            collector.add_dir(saved_dir)
+            return
+
+        for sub_dir in (x for x in saved_dir.glob("*") if x.is_dir()):
+            if sub_dir.name.lower() not in self._saved_keeps:
+                collector.add_dir(sub_dir)
+
+    def _detect_locked_files(self, root_dir):
+        proc = sp.Popen(
+            ("wmic.exe", "process", "get", "executablepath"),
+            stdout=sp.PIPE, stderr=sp.DEVNULL
+        )
+
+        ret = False
+        for line in proc.stdout.readlines():
+            line = line.strip().decode()
+            if line and Path(line).parent.is_relative_to(root_dir):
+                ret = line
+                break
+
+        proc.stdout.close()
+        proc.wait()
+        return ret
+
+    def main(self):
+        self._saved_keeps = {x.strip().lower() for x in self.args.savedkeeps.split(",")}
+
+        ue_context = unreal.Context(os.getcwd())
+        branch = ue_context.get_branch(must_exist=True)
+        engine_dir = ue_context.get_engine().get_dir()
+
+        if os.name == "nt":
+            self.print_info("Checking running processes")
+            if running_exe := self._detect_locked_files(branch.get_dir()):
+                raise RuntimeError(f"Not cleaning because '{running_exe}' is running")
+                return False
+
+        self.print_info("Finding directories and files to clean")
+        root_dirs = [
+            engine_dir,
+            *(x for x in engine_dir.glob("Programs/*") if x.is_dir()),
+        ]
+
+        print("Enumerating...", end="")
+        rg_args = (
+            "rg",
+            "--files",
+            "--path-separator=/",
+            "--no-ignore",
+            "-g*.uplugin",
+            "-g*.uproject",
+            str(branch.get_dir()),
+        )
+        rg = sp.Popen(rg_args, stdout=sp.PIPE, stderr=sp.DEVNULL)
+        for line in rg.stdout.readlines():
+            path = Path(line.decode().rstrip()).parent
+            root_dirs.append(path)
+        rg.wait()
+        print("\r", len(root_dirs), " uproject/uplugin roots found", sep="")
+
+        collector = _Collector()
+        clean_handlers = {
+            "Intermediate"     : lambda x,y: x.add_dir(y),
+            "DerivedDataCache" : lambda x,y: x.add_dir(y),
+            "Binaries"         : lambda x,y: x.add_dir(y, remove_carefully=True),
+            "Saved"            : self._append_saved,
+        }
+        for root_dir in root_dirs:
+            for dir_name, clean_handler in clean_handlers.items():
+                candidate = root_dir / dir_name
+                if candidate.is_dir():
+                    clean_handler(collector, candidate)
+
+        # Ask Perforce which files shouldn't be deleted
+        print("Asking Perforce what's synced...", end="", flush=True)
+        specs = (str(x) + "/..." for x in collector.read_careful_dirs())
+        p4_have = P4.have(specs)
+        for item in p4_have.read(on_error=False):
+            path = Path(item.path)
+            collector.pin_file(path)
+        print("done (", collector.get_pinned_count(), " files)", sep="")
+
+        if self.args.dryrun:
+            collector.dispatch(_DryRun())
+            return
+
+        self.print_info("Cleaning")
+        actuator = _Cleaner(branch.get_dir())
+        collector.dispatch(actuator)
+
+
+
+#-------------------------------------------------------------------------------
+class Reset(flow.cmd.Cmd):
+    """ Reconciles a branch to make it match the depot. Use with caution; this is
+    a destructive action and involes removing and rewriting files! """
+    thorough = flow.cmd.Opt(False, "Compare digests instead of file-modified time")
+
+    def _get_reset_paths(self):
+        try:
+            ue_context = unreal.Context(os.getcwd())
+            if not (branch := ue_context.get_branch()):
+                return None
+        except EnvironmentError:
+            return None
+
+        root_dir = branch.get_dir()
+        ret = {
+            # root_dir / "*", # to much scope to go wrong here and it's so few files
+            root_dir / "Template/...",
+            root_dir / "Engine/...",
+        }
+        for uproj_path in branch.read_projects():
+            ret.add(uproj_path.parent.absolute() / "...")
+
+        return ret
+
+    def main(self):
+        # Confirm the user wants to really do a reset
+        self.print_warning("Destructive action!")
+        if self.is_interactive():
+            while True:
+                c = input("Are you sure you want to continue [yn] ?")
+                if c.lower() == "y": break
+                if c.lower() == "n": return False
+
+        # Run the reconcile
+        args = ("reconcile", "-wade",)
+        if not self.args.thorough:
+            args = (*args, "--modtime")
+
+        if reset_paths := self._get_reset_paths():
+            args = (*args, *(str(x) for x in reset_paths))
+
+        exec_context = self.get_exec_context()
+        cmd = exec_context.create_runnable("p4", *args)
+        return cmd.run()
diff --git a/Engine/Extras/ushell/channels/unreal/perforce/cmds/gui.py b/Engine/Extras/ushell/channels/unreal/perforce/cmds/gui.py
new file mode 100644
--- /dev/null
+++ b/Engine/Extras/ushell/channels/unreal/perforce/cmds/gui.py
@@ -0,0 +1,72 @@
+# Copyright Epic Games, Inc. All Rights Reserved.
+
+import os
+import p4utils
+import flow.cmd
+import subprocess
+from pathlib import Path
+
+#-------------------------------------------------------------------------------
+def _check_for_existing_nt(client):
+    client = client.encode()
+
+    wmic = ("wmic.exe", "process", "get", "name,commandline")
+    proc = subprocess.Popen(wmic, stdout=subprocess.PIPE)
+    for line in iter(proc.stdout.readline, b""):
+        if b"p4v.exe" in line and client in line:
+            return True
+    proc.stdout.close()
+    proc.wait()
+
+_check_for_existing = _check_for_existing_nt if os.name == "nt" else lambda x: ""
+
+
+
+#-------------------------------------------------------------------------------
+class Gui(flow.cmd.Cmd):
+    """ Opens the clientspec for the current directory in P4V """
+    filename = flow.cmd.Arg(Path(), "Path to select in the workspace/depot tree")
+    p4vargs  = flow.cmd.Arg([str], "Additional arguments to pass to P4V")
+    def main(self):
+        self.print_info("Fetching Perforce info")
+
+        username = p4utils.login()
+
+        cwd = os.getcwd()
+        print(f"Finding client for '{cwd}'; ", end="")
+        client = p4utils.get_client_from_dir(cwd, username)
+        if not client:
+            print()
+            self.print_error("No client found")
+            return False
+
+        client, client_dir = client
+        print(client)
+
+        p4_port = p4utils.get_p4_set("P4PORT")
+        if not p4_port:
+            self.print_error("P4PORT must be set to start P4V")
+            return False
+        print("Using P4PORT", p4_port)
+
+        self.print_info("Starting P4V")
+        if _check_for_existing(client):
+            print("P4V with client", client, "is already open")
+            return
+
+        print("Client root;", client_dir)
+
+        args = (
+            "-p", p4_port,
+            "-u", username,
+            "-c", client,
+            "-t", "pending",
+        )
+
+        if self.args.filename.is_file():
+            filename = self.args.filename.resolve()
+            args = (*args, "-s", str(filename))
+
+        print("p4v", *args)
+        subprocess.Popen(("p4v", *args, *self.args.p4vargs), cwd=client_dir)
+        print("Done!")
diff --git a/Engine/Extras/ushell/channels/unreal/perforce/cmds/mergedown.py b/Engine/Extras/ushell/channels/unreal/perforce/cmds/mergedown.py
new file mode 100644
--- /dev/null
+++ b/Engine/Extras/ushell/channels/unreal/perforce/cmds/mergedown.py
@@ -0,0 +1,173 @@
+# Copyright Epic Games, Inc. All Rights Reserved.
+
+import p4utils
+import flow.cmd
+from peafour import P4
+
+#-------------------------------------------------------------------------------
+class _TtyCounter(object):
+    def __init__(self, leader, suffix):
+        self._suffix = suffix
+        self._count = 0
+        self._leader = leader
+        self.flush(True)
+
+    def inc(self):
+        self._count += 1
+        self.flush()
+
+    def flush(self, inc_leader=False):
+        if inc_leader:
+            print(self._leader, end="")
+        out = str(self._count) + self._suffix
+        print(out, "\b" * len(out), sep="", end="")
+        return self
+
+
+
+#-------------------------------------------------------------------------------
+class MergeDown(flow.cmd.Cmd):
+    """ Merge down from the parent stream. """
+    changelist    = flow.cmd.Arg("", "Changelist to merge down at")
+    path          = flow.cmd.Opt("", "Limit mergedown to a sub-tree; --path=Engine/Source")
+    maxscanrows   = flow.cmd.Opt(0, "Set the max scan rows with integrating")
+    dryrun        = flow.cmd.Opt(False, "Only pretend to do work")
+    noautoresolve = flow.cmd.Opt(False, "Do not automatically resolve")
+
+    @flow.cmd.Cmd.summarise
+    def main(self):
+        p4utils.login()
+
+        # Determine the streams involved
+        self.print_info("Getting stream info")
+        info = P4.info()
+        stream = getattr(info, "clientStream", None)
+        if not stream:
+            raise EnvironmentError(f"Client '{info.clientName}' is not a stream")
+
+        # Get the parent, skipping past virtual streams.
+        stream_info = P4.stream(stream, o=True).run()
+        parent = stream_info.Parent
+        while True:
+            stream_info = P4.stream(parent, o=True).run()
+            if stream_info.Type != "virtual":
+                break
+            parent = stream_info.Parent
+
+        if self.args.changelist:
+            parent_cl = self.args.changelist
+        else:
+            parent_cl = P4.changes(parent + "/...", m=1, s="submitted").run()
+            parent_cl = parent_cl.change
+
+        print(f"Target: {stream} (as {info.clientName})")
+        print(f"Parent: {parent}")
+        print(f"Change: {parent_cl}")
+
+        # Create a changelist to merge into
+        if self.args.path:
+            spec_stem = f"/{self.args.path}/...@".replace("\\", "/")
+            print("Filter:", stream + spec_stem + parent_cl)
+        else:
+            spec_stem = "/...@"
+
+        # Create a changelist to merge into
+        def _create_cl(desc):
+            cl_spec = {
+                "Change" : "new",
+                "Description" : desc,
+            }
+            P4.change(i=True).run(input_data=cl_spec)
+            return P4.changes(c=info.clientName, m=1, s="pending").change
+
+        self.print_info("Creating target changelist")
+        cl_desc = f"Merging {parent} @ {parent_cl} to {stream} ({info.clientName})"
+        cl_desc += "\n\n#ushell-mergedown"
+        dest_cl = _create_cl(cl_desc)
+        print(dest_cl)
+
+        # Do the merge
+        self.print_info("Merging")
+
+        path_args = (
+            stream + spec_stem + parent_cl,
+        )
+
+        int_args = {
+            "c" : dest_cl,
+            "n" : self.args.dryrun,
+            "S" : stream,
+            "r" : True,
+        }
+
+        def on_error(data):
+            msg = data.data.strip()
+            print(end="\r")
+            self.print_error(msg)
+            file_count.flush(True)
+            file_count.inc()
+
+        file_count = _TtyCounter("Running 'p4 integrate'; ", " files")
+        p4 = P4(zmaxscanrows=str(self.args.maxscanrows or 1000000000))
+        integrate = p4.integrate(*path_args, **int_args)
+        for item in integrate.read(on_error=on_error):
+            path = getattr(item, "clientFile", None)
+            if path:
+                file_count.inc()
+        print("\nDone!")
+
+        if self.args.dryrun:
+            return
+
+        # Resolve safely
+        def do_resolve(cl, *modes):
+            mode_args = {x:True for x in modes}
+            yield from P4.resolve(c=cl, **mode_args).read(on_error=False)
+
+        self.print_info("Resolving")
+        file_count = _TtyCounter("Safely; ", " files")
+        for item in do_resolve(dest_cl, "as"):
+            path = getattr(item, "clientFile", None)
+            if path:
+                file_count.inc()
+
+        # Move files that didn't resolve to a review changelist
+        def create_review_cl():
+            lines = ("AUTO MERGED/UNRESOLVED", cl_desc, "#nocheckin")
+            review_cl_desc = "\n\n".join(lines)
+            review_cl = _create_cl(review_cl_desc)
+            print("Moving unresolved to " + review_cl)
+            return review_cl
+
+        review_cl = None
+        try:
+            for item in do_resolve(dest_cl, "n"):
+                review_cl = review_cl or create_review_cl()
+                path = getattr(item, "clientFile", None)
+                if path:
+                    P4.reopen(path, c=review_cl).run()
+        except P4.Error:
+            pass
+
+        # Perhaps there's nothing to review? Hooray!
+        if not review_cl:
+            return
+
+        # Make sure moved-pairs are in the same changelist
+        def read_moved_pairs():
+            for item in P4.opened(c=review_cl):
+                move_pair = getattr(item, "movedFile", None)
+                if move_pair:
+                    yield move_pair
+
+        file_count = _TtyCounter("Collecting moved pairs; ", " files")
+        for item in P4.reopen(read_moved_pairs(), c=review_cl):
+            file_count.inc()
+
+        # Resolve what we can automatically.
+        if not self.args.noautoresolve:
+            file_count = _TtyCounter("Automatically; ", " files")
+            for item in do_resolve(review_cl, "am"):
+                path = getattr(item, "clientFile", None)
+                if path:
+                    file_count.inc()
diff --git a/Engine/Extras/ushell/channels/unreal/perforce/cmds/project.py b/Engine/Extras/ushell/channels/unreal/perforce/cmds/project.py
new file mode 100644
--- /dev/null
+++ b/Engine/Extras/ushell/channels/unreal/perforce/cmds/project.py
@@ -0,0 +1,27 @@
+# Copyright Epic Games, Inc. All Rights Reserved.
+
+import os
+import p4utils
+import flow.cmd
+
+#-------------------------------------------------------------------------------
+class Change(flow.cmd.Cmd):
+    def main(self):
+        # Call the main project command (inheritance happens in describe.flow.py)
+        primary_dir = super().main()
+        if not primary_dir:
+            return False
+
+        self.print_info("Perforce environment")
+        try:
+            p4config = p4utils.ensure_p4config(primary_dir)
+            if p4config:
+                path, created = p4config
+                print("Status:", "created" if created else "exists")
+                print("Path:", path)
+            else:
+                self.print_warning("Unabled to find a client from", primary_dir)
+        except EnvironmentError as e:
+            self.print_warning(e)
+
+        return True
