Updated ushell [26875160]

diff --git a/Engine/Extras/ushell/channels/unreal/perforce/cmds/switch.py b/Engine/Extras/ushell/channels/unreal/perforce/cmds/switch.py
new file mode 100644
--- /dev/null
+++ b/Engine/Extras/ushell/channels/unreal/perforce/cmds/switch.py
@@ -0,0 +1,289 @@
+# Copyright Epic Games, Inc. All Rights Reserved.
+
+import os
+import fzf
+import unreal
+import p4utils
+import flow.cmd
+import subprocess
+from peafour import P4
+
+#-------------------------------------------------------------------------------
+class _RestorePoint(object):
+    def __init__(self):
+        self._files = {}
+        self._shelf = None
+
+    def __del__(self):
+        if not self._files:
+            return
+
+        print("Attemping to restore changelists")
+        if self._shelf:
+            print("Unshelving...")
+            for item in P4.unshelve(s=self._shelf).read(on_error=False):
+                pass
+
+        print("Restoring...")
+        for cl, paths in self._files.items():
+            print(" ", cl)
+            for item in P4.reopen(paths, c=cl).read(on_error=False):
+                print("   ", item.depotFile)
+
+    def set_shelf(self, shelf):
+        self._shelf = shelf
+
+    def add_file(self, changelist, path):
+        out = self._files.setdefault(changelist, [])
+        out.append(path)
+
+    def clear(self):
+        self._files.clear()
+
+#-------------------------------------------------------------------------------
+class _SwitchCmd(flow.cmd.Cmd):
+    def _read_streams(self):
+        depot = self._depot.lower()
+        yield from (x for x in P4.streams() if x.Stream.lower().startswith(depot))
+
+    def main(self):
+        self.print_info("Perforce environment")
+
+        # Ensure there's a valid branch and Perforce ticket
+        branch = unreal.Context(".").get_branch()
+        if not branch:
+            raise EnvironmentError("Unable to find a valid branch")
+        os.chdir(branch.get_dir())
+
+        self._username = p4utils.login()
+
+        # Find out where we are and the stream
+        self._depot = None
+
+        info = P4.info().run()
+        self._src_stream = getattr(info, "clientStream", None)
+        if self._src_stream:
+            self._depot = "//" + self._src_stream[2:].split("/")[0] + "/"
+
+        print("Client:", info.clientName)
+        print("  User:", self._username)
+        print("Stream:", self._src_stream)
+        print(" Depot:", self._depot)
+        if not self._src_stream:
+            self.print_error(info.clientName, "is not a stream-based client")
+            return False
+
+        return self._main_impl()
+
+#-------------------------------------------------------------------------------
+class List(_SwitchCmd):
+    """ Prints a tree of available streams relevant to the current branch """
+
+    def _main_impl(self):
+        self.print_info("Available streams")
+
+        streams = {}
+        for item in self._read_streams():
+            stream = streams.setdefault(item.Stream, [None, []])
+            stream[0] = item
+            parent = streams.setdefault(item.Parent, [None, []])
+            parent[1].append(item.Stream)
+
+        def print_stream(name, depth=0):
+            item, children = streams[name]
+            prefix = ("| " * depth) + "+ "
+            dots = "." * (64 - len(item.Name) - (2 * depth))
+            print(flow.cmd.text.grey(prefix), item.Name, sep="", end=" ")
+            print(flow.cmd.text.grey(dots), end=" ")
+            print(item.Type)
+            for child in sorted(children):
+                print_stream(child, depth + 1)
+
+        _, roots = streams.get("none", (None, None))
+        if not roots:
+            print("None?")
+            return False
+        for root in roots:
+            print_stream(root)
+
+#-------------------------------------------------------------------------------
+class Switch(_SwitchCmd):
+    """ Switch between streams and integrate any open files across to the new
+    stream. If no stream name is provided then the user will be prompted to
+    select a stream from a list. Operates on the current directory's branch.
+
+    There maybe occasions where a stream is only partially synced like a temporary
+    client used for cherrypicking. The '--haveonly' can be used in cases liek this
+    to avoid syncing the whole branch post-switch and instead just update the
+    files synced prior to switching. Use with extraordinary caution!"""
+    stream      = flow.cmd.Arg("", "Name of the stream to switch to")
+    changelist  = flow.cmd.Arg(-1, "The changelist to sync to ('head' if unspecified)")
+    haveonly    = flow.cmd.Opt(False, "Only switch files synced prior to the switch")
+    saferesolve = flow.cmd.Opt(False, "Resolve safely and not automatically")
+
+    def _select_stream(self):
+        self.print_info("Stream select")
+        print("Enter the stream to switch to (fuzzy, move with arrows, enter to select)...", end="")
+        stream_iter = (x.Name for x in self._read_streams())
+        for reply in fzf.run(stream_iter, height=10, prompt=self._depot):
+            return reply
+
+    @flow.cmd.Cmd.summarise
+    def _main_impl(self):
+        # Work out the destination stream.
+        if not self.args.stream:
+            dest_stream = self._select_stream()
+            if not dest_stream:
+                self.print_error("No stream selected")
+                return False
+        else:
+            dest_stream = self.args.stream
+        dest_stream = self._depot + dest_stream
+
+        # Check for the destination stream and correct its name
+        self.print_info("Checking destination")
+        stream = P4.streams(dest_stream, m=1).run()
+        dest_stream = stream.Stream
+        print("Stream:", dest_stream)
+        print("Parent:", stream.Parent)
+        print("  Type:", stream.Type)
+        self._dest_stream = dest_stream
+
+        if self._src_stream.casefold() == self._dest_stream.casefold():
+            self.print_warning("Already on stream", self._dest_stream)
+            return
+
+        # Move the user's opened files into a changelist
+        self.print_info("Shelving and reverting open files")
+        self._get_opened_files()
+
+        if len(self._opened_files):
+            shelf_cl = self._shelve_revert()
+
+        # Do the switch
+        P4.client(S=self._dest_stream, s=True).run()
+        if self.args.haveonly:
+            self._do_have_table_update()
+        else:
+            self._do_sync()
+
+        # Unshelve to restore user's open files
+        if len(self._opened_files):
+            self._restore_changes(shelf_cl)
+
+    def _do_have_table_update(self):
+        # This probably isn't correct if stream specs diff significantly
+        # betweeen source and destination.
+
+        def read_have_table():
+            for item in P4.have():
+                depot_path = item.depotFile
+                yield self._dest_stream + depot_path[len(self._src_stream):]
+
+        print("Updating have table; ", end="")
+        for i, item in enumerate(P4.sync(read_have_table()).read(on_error=False)):
+            print(i, "\b" * len(str(i)), end="", sep="")
+        for x in P4.sync(self._src_stream + "/...#have").read(on_error=False):
+            pass
+        print()
+
+    def _do_sync(self):
+        self.print_info("Syncing")
+
+        self.print_warning("If the sync is interrupted the branch will be a mix of files synced pre-")
+        self.print_warning("and post-switch. Should this happen `.p4 sync` can be used to recover but")
+        self.print_warning("any open files at the time of the switch will remain shelved.")
+
+        sync_cmd = ("_p4", "sync", "--all", "--noresolve", "--nosummary")
+        if self.args.changelist >= 0:
+            sync_cmd = (*sync_cmd, str(self.args.changelist))
+
+        popen_kwargs = {}
+        if os.name == "nt":
+            popen_kwargs["creationflags"] = subprocess.CREATE_NEW_PROCESS_GROUP
+
+        proc = subprocess.Popen(sync_cmd, **popen_kwargs)
+        proc.wait()
+
+    def _get_opened_files(self):
+        self._restore_point = _RestorePoint();
+        self._opened_files = {}
+        src_len = len(self._src_stream) + 1
+        for item in P4.opened():
+            if not item.depotFile.startswith(self._src_stream):
+                raise ValueError(f"Open file {item.depotFile} is not from the current stream {self._src_stream}")
+
+            base_path = item.depotFile[src_len:]
+
+            out = self._opened_files.setdefault(item.change, [])
+            out.append(base_path)
+
+            self._restore_point.add_file(item.change, base_path)
+
+            prefix = item.change if len(out) == 1 else ""
+            print("%-9s" % prefix, base_path, sep=": ")
+
+    def _shelve_revert(self):
+        def _for_each_open(command, info_text):
+            def read_open_files():
+                for x in self._opened_files.values():
+                    yield from x
+
+            print(info_text, "." * (22 - len(info_text)), end=" ")
+            runner = getattr(P4, command)
+            runner = runner(read_open_files(), c=shelf_cl)
+            for i, item in enumerate(runner):
+                print(i, "\b" * len(str(i)), sep="", end="")
+            print("done")
+
+        cl_desc = f"'.p4 switch {self.args.stream}' backup of opened files from {self._src_stream}\n\n#ushell-switch"
+        cl_spec = {"Change": "new", "Description": cl_desc}
+        P4.change(i=True).run(input_data=cl_spec)
+        shelf_cl = P4.changes(u=self._username, m=1, s="pending").change
+
+        _for_each_open("reopen", "Moving changelist")
+        _for_each_open("shelve", "Shelving")
+        _for_each_open("revert", "Reverting")
+
+        self._restore_point.set_shelf(shelf_cl)
+
+        return shelf_cl
+
+    def _restore_changes(self, shelf_cl):
+        # Unshelve each file into its original CL, integrating between streams
+        # at the same time.
+        self.print_info("Restoring changes")
+        branch_spec = p4utils.TempBranchSpec("switch", self._username, self._src_stream, self._dest_stream)
+        print("Branch spec:", branch_spec)
+        for orig_cl, files in self._opened_files.items():
+            print(orig_cl)
+            orig_cl = None if orig_cl == "default" else orig_cl
+            unshelve = P4.unshelve(files, b=branch_spec, s=shelf_cl, c=orig_cl)
+            for item in unshelve.read(on_error=False):
+                print("", item.depotFile)
+        del branch_spec
+        self._restore_point.clear()
+
+        # Resolve files
+        def print_error(error):
+            print("\r", end="")
+            self.print_error(error.data.rstrip())
+
+        resolve_count = 0
+        resolve_args = {
+            "as" : self.args.saferesolve,
+            "am" : not self.args.saferesolve,
+        }
+        resolve = P4.resolve(**resolve_args)
+        for item in resolve.read(on_error=print_error):
+            if getattr(item, "clientFile", None):
+                resolve_count += 1
+                print("\r" + str(resolve_count), "file(s) resolved", end="")
+        if resolve_count:
+            print("")
+
+        # Inform the user about conflicts
+        resolve = P4.resolve(n=True)
+        for item in resolve.read(on_error=False):
+            if name := getattr(item, "fromFile", None):
+                self.print_error(name[len(self._src_stream) + 1:])
diff --git a/Engine/Extras/ushell/channels/unreal/perforce/cmds/sync.py b/Engine/Extras/ushell/channels/unreal/perforce/cmds/sync.py
new file mode 100644
--- /dev/null
+++ b/Engine/Extras/ushell/channels/unreal/perforce/cmds/sync.py
@@ -0,0 +1,379 @@
+# Copyright Epic Games, Inc. All Rights Reserved.
+
+import os
+import unreal
+import p4utils
+import flow.cmd
+import unrealcmd
+from peafour import P4
+from pathlib import Path
+
+
+#-------------------------------------------------------------------------------
+class _SyncBase(unrealcmd.Cmd):
+    def _write_p4sync_txt_header(self, out):
+        out.write("# lines starting with a '#' are comments\n")
+        out.write("# lines prefixed with '-' are excluded from your sync\n")
+        out.write("# -.../Android/...\n")
+        out.write("# -/Engine/Source/...\n")
+        out.write("# -*.uasset\n")
+
+
+
+#-------------------------------------------------------------------------------
+class Sync(_SyncBase):
+    """ Syncs the current project and engine directories to the given changelist
+    (or latest if none is provided). If the '--all' option is specified then the
+    branch will be searched locally for existing .uproject files, scheduling each
+    one to be synced.
+
+    The sync can be filtered with a .p4sync.txt file. Lines prefixed with a '-'
+    (e.g. "-.../SootySweep/...") will be excluded from the sync and anything
+    already synced is de-synced. The command will read .p4sync.txt files from two
+    locations;
+
+        1. <branch_root>/.p4sync.txt
+        2. ~/.ushell/.p4sync.txt (where ~ is USERPROFILE on Windows)
+
+    Quick edit the branch's .p4sync.txt file with '.p4 sync edit'. An example
+    of a .p4sync.txt file is as follows;
+
+        # a comment
+        -.../Android/...
+        -/Engine/Source/...
+        -*.uasset """
+    changelist = unrealcmd.Arg(-1, "The changelist to sync to ('now' if unspecified)")
+    noresolve  = unrealcmd.Opt(False, "Do not run 'p4 resolve -am' after syncing")
+    dryrun     = unrealcmd.Opt(False, "Only pretend to do the sync")
+    all        = unrealcmd.Opt(False, "Sync all the branch's projects found locally")
+    addprojs   = unrealcmd.Opt("", "Comma-separated names of additional projects to sync")
+    clobber    = unrealcmd.Opt(False, "Clobber writable files when syncing")
+    nosummary  = unrealcmd.Opt(False, "Do not print the result-time summary at the end")
+    echo       = unrealcmd.Opt(False, "Echo depot paths as they are synced")
+
+    def complete_addprojs(self, prefix):
+        # Fake a _local_root
+        local_root = Path(os.getcwd())
+        for parent in local_root.parents:
+            if (parent / "GenerateProjectFiles.bat").is_file():
+                local_root = parent
+                break
+        self._local_root = local_root
+
+        # Now we can fetch an approximate context and list it's projects
+        ue_context = self._try_get_unreal_context()
+        if not ue_context:
+            return
+
+        branch = ue_context.get_branch()
+        if not branch:
+            return
+
+        return (x.stem for x in  branch.read_projects())
+
+    @unrealcmd.Cmd.summarise
+    def _main_summarised(self):
+        return self._main_impl()
+
+    def main(self):
+        self._client_spec_restore = None
+        try:
+            if self.args.nosummary:
+                return self._main_impl()
+            else:
+                return self._main_summarised()
+        finally:
+            try:
+                if self._client_spec_restore:
+                    print("Restoring client spec")
+                    client_spec = self._client_spec_restore
+                    client_spec["Options"] = client_spec["Options"].replace("clobber", "noclobber")
+                    client_spec["Description"] = client_spec["Description"].replace("{ushell_clobber_patch}", "")
+                    P4.client(i=True).run(input_data=client_spec)
+            except:
+                pass
+
+    def _setup(self):
+        self.print_info("Perforce environment")
+
+        # Check there's a valid Perforce environment.
+        username = p4utils.login()
+        if not p4utils.ensure_p4config():
+            self.print_warning("Unable to establish a P4CONFIG")
+
+        # Get some info about the Perforce environment and show it to the user
+        info = P4.info().run()
+        print("Client:", info.clientName)
+        print("  User:", info.userName)
+        print("Server:", getattr(info, "proxyAddress", info.serverAddress))
+
+        # Inform the user if Perforce didn't find the client.
+        if info.clientName == "*unknown*":
+            client_name = p4utils.get_p4_set("P4CLIENT")
+            _, p4config_name = p4utils.has_p4config(".")
+            raise EnvironmentError(f"Client '{client_name}' not found. Please check P4CLIENT setting in '{p4config_name}'")
+
+        # So that P4.where can succeed we sync one known file first. This also
+        # ensures we can accomodate an unsynced stream switch.
+        for x in P4.sync(f"//{info.clientName}/GenerateProjectFiles.bat").read(on_error=False):
+            pass
+
+        # Find the root of the current branch
+        self.print_info("Discovering branch root")
+        branch_root = p4utils.get_branch_root(f"//{info.clientName}/Engine/Source")
+        print("Branch root:", branch_root)
+
+        # Map branch root somewhere on the local file system
+        local_root = P4.where(branch_root + "X").path
+        local_root = local_root[:-1] # to strip 'X'
+        print("Local root:", local_root)
+
+        self._info = info
+        self._branch_root = branch_root
+        self._local_root = local_root
+
+    def _try_get_unreal_context(self):
+        try:
+            ue_context = self.get_unreal_context()
+            branch = ue_context.get_branch()
+            # If branch doesn't match os.getcwd() then ditch it
+            if not (branch and branch.get_dir().samefile(self._local_root)):
+                raise EnvironmentError
+        except EnvironmentError:
+            try:
+                cwd = os.getcwd()
+                ue_context = unreal.Context(cwd)
+            except EnvironmentError:
+                ue_context = None
+
+        return ue_context
+
+    def _add_paths(self, syncer):
+        # Add the set of paths that all syncs should include
+        syncer.add_path(self._local_root + "*")
+        syncer.add_path(self._local_root + "Engine/...")
+        templates = self._local_root + "Templates/..."
+        if P4.files(templates, m=1).run(on_error=lambda x: None) is not None:
+            syncer.add_path(templates)
+
+        # If we've a valid context by this point we can try and use it.
+        glob_for_projects = False
+        self._current_cl = 0
+        if ue_context := self._try_get_unreal_context():
+            project = ue_context.get_project()
+            if self.args.all or not project:
+                if branch := ue_context.get_branch():
+                    project_count = 0
+                    self.print_info("Syncing all known projects")
+                    for uproj_path in branch.read_projects():
+                        print(uproj_path.stem)
+                        syncer.add_path(str(uproj_path.parent) + "/...")
+                        project_count += 1
+
+                    # If we have somehow managed to not find any projects then
+                    # fallback to globbing for them.
+                    if not project_count:
+                        print("No projects found via .uprojectdirs")
+                        print("Falling back to a glob")
+                        glob_for_projects = True
+            else:
+                # By default the active project is synced
+                self.print_info("Single project sync")
+                print("Project:", project.get_name())
+                syncer.add_path(str(project.get_dir()) + "/...")
+
+            # Extra projects
+            if self.args.addprojs and not self.args.all:
+                add_projects = self.args.addprojs.replace("/", ",")
+                add_projects = (x.strip() for x in add_projects.split(","))
+                add_projects = {x for x in add_projects if x}
+                known_projects = list(ue_context.get_branch().read_projects())
+                known_projects = {x.stem.lower():x for x in known_projects}
+                self.print_info("Additional projects to sync;")
+                for add_project in add_projects:
+                    print(add_project, ": ", sep="", end="")
+                    add_project = add_project.lower()
+                    if add_project not in known_projects:
+                        print("not found")
+                        continue
+
+                    add_project = known_projects[add_project]
+                    add_project = add_project.parent
+                    syncer.add_path(str(add_project) + "/...")
+                    print(add_project)
+
+            engine_info = ue_context.get_engine().get_info()
+            self._current_cl = engine_info.get("Changelist", 0)
+        else:
+            glob_for_projects = True
+
+        if glob_for_projects:
+            # There does not appear to be a fully formed branch so we will infer
+            # `--all` here on behalf of the user.
+            self.print_info("Syncing all projects by **/.uproject")
+            for uproj_path in Path(self._local_root).glob("**/*.uproject"):
+                print(uproj_path.stem)
+                syncer.add_path(str(uproj_path.parent) + "/...")
+
+    def _main_impl(self):
+        self._setup()
+
+        # Determine the changelist to sync
+        sync_cl = self.args.changelist
+        if sync_cl < 0:
+            sync_cl = int(P4.changes(self._branch_root + "...", m=1).change)
+
+        # Remove "noclobber" from the user's client spec
+        client = P4.client(o=True).run()
+        client_spec = client.as_dict()
+        client_spec.setdefault("Description", "")
+
+        if self.args.clobber:
+            self.print_info("Checking for 'noclobber'")
+            if "noclobber" in client_spec["Options"]:
+                client_spec["Options"] = client_spec["Options"].replace("noclobber", "clobber")
+                client_spec["Description"] += "{ushell_clobber_patch}"
+                self._client_spec_restore = client_spec.copy()
+                if not self.args.dryrun or True:
+                    print(f"Patching {client.Client} with 'clobber'")
+                    P4.client(i=True).run(input_data=client_spec)
+            else:
+                print("Clobbering is already active")
+
+        if not self._client_spec_restore:
+            if "{ushell_clobber_patch}" in client_spec["Description"]:
+                if "noclobber" not in client_spec["Options"]:
+                    self._client_spec_restore = client_spec.copy()
+
+        # Add the paths we always want to sync
+        syncer = p4utils.Syncer()
+        self._add_paths(syncer)
+
+        # Load and parse the .p4sync.txt file
+        self._apply_p4sync_txt(syncer)
+
+        version_cl = 0
+        build_ver_path = self._local_root + "Engine/Build/Build.version"
+        try:
+            # Special case to force sync Build.version. It can get easily modified
+            # without Perforce's knowledge, complicating the sync.
+            if not self.args.dryrun:
+                P4.sync(build_ver_path + "@" + str(sync_cl), qf=True).run(on_error=False)
+
+            # GO!
+            self.print_info("Scheduling sync")
+            print("Changelist:", sync_cl, f"(was {self._current_cl})")
+            print("Requesting... ", end="")
+            syncer.schedule(sync_cl)
+
+            self.print_info("Syncing")
+            ok = syncer.sync(dryrun=self.args.dryrun, echo=self.args.echo)
+
+            if self.args.dryrun or not ok:
+                return ok
+
+            # Sync succeeded, update cl for build.version even if something goes wrong with resolving
+            version_cl = sync_cl
+
+            # Auto-resolve on behalf of the user.
+            if not self.args.noresolve:
+                conflicts = set()
+                self.print_info("Resolving")
+                for item in P4.resolve(am=True).read(on_error=False):
+                    path = getattr(item, "fromFile", None)
+                    if not path:
+                        continue
+
+                    path = path[len(self._branch_root):]
+                    if getattr(item, "how", None):
+                        conflicts.remove(path)
+                        print(path)
+                    else:
+                        conflicts.add(path)
+
+                for conflict in conflicts:
+                    print(flow.cmd.text.light_red(conflict))
+
+        except KeyboardInterrupt:
+            print()
+            self.print_warning(f"Sync interrupted! Writing build.version to CL {version_cl}")
+            return False
+        finally:
+            # Record the synced changelist in Build.version
+            with open(build_ver_path, "r") as x:
+                lines = list(x.readlines())
+
+            import stat
+            build_ver_mode = os.stat(build_ver_path).st_mode
+            os.chmod(build_ver_path, build_ver_mode|stat.S_IWRITE)
+
+            with open(build_ver_path, "w") as x:
+                for line in lines:
+                    if r'"Changelist"' in line:
+                        line = line.split(":", 2)
+                        line = line[0] + f": {version_cl},\n"
+                    elif r'"BranchName"' in line:
+                        line = "\t\"BranchName\": \"X\"\n"
+                        line = line.replace("X", self._branch_root[:-1].replace("/", "+"))
+                    x.write(line)
+
+    def _apply_p4sync_txt(self, syncer):
+        def impl(path):
+            print("Source:", os.path.normpath(path), end="")
+            try:
+                sync_config = open(path, "rt")
+                print()
+            except:
+                print(" ... not found")
+                return
+
+            def read_exclusions():
+                for line in map(str.strip, sync_config):
+                    if line.startswith("-"):    yield line[1:]
+                    elif line.startswith("$-"): yield line[2:]
+
+            for i, line in enumerate(read_exclusions()):
+                view = None
+                if line.startswith("*."):    view = ".../" + line
+                elif line.startswith("/"):   view = line[1:]
+                elif line.startswith("..."): view = line
+
+                print("       %2d" % i, "exclude", end=" ")
+
+                if view and (view.count("/") or "/*." in view or view.startswith("*.")):
+                    view = self._branch_root + view
+                    syncer.add_exclude(view)
+                    print(view)
+                else:
+                    view = view or line
+                    print(flow.cmd.text.light_yellow(view + " (ill-formed)"))
+
+            sync_config.close()
+
+        self.print_info("Applying .p4sync.txt")
+        for dir in (self.get_home_dir(), self._local_root):
+            impl(dir + ".p4sync.txt")
+
+
+
+#-------------------------------------------------------------------------------
+class Edit(_SyncBase):
+    """ Opens .p4sync.txt in an editor. The editor is selected from environment
+    variables P4EDITOR, GIT_EDITOR, and the system default editor. """
+    def main(self):
+        username = p4utils.login()
+
+        cwd = os.getcwd()
+        client = p4utils.get_client_from_dir(cwd, username)
+        if not client:
+            raise EnvironmentError(f"Unable to establish the clientspec from '{cwd}'")
+
+        _, root_dir = client
+        path = Path(root_dir) / ".p4sync.txt"
+
+        if not path.is_file():
+            with path.open("wt") as out:
+                self._write_p4sync_txt_header(out)
+
+        print("Editing", path)
+        self.edit_file(path)
diff --git a/Engine/Extras/ushell/channels/unreal/perforce/cmds/who.py b/Engine/Extras/ushell/channels/unreal/perforce/cmds/who.py
new file mode 100644
--- /dev/null
+++ b/Engine/Extras/ushell/channels/unreal/perforce/cmds/who.py
@@ -0,0 +1,356 @@
+# Copyright Epic Games, Inc. All Rights Reserved.
+
+import re
+import os
+import fzf
+import time
+import shutil
+import p4utils
+import tempfile
+import flow.cmd
+from peafour import P4
+import subprocess as sp
+
+#-------------------------------------------------------------------------------
+def _as_p4_path(path_rev):
+    path, rev = map(str, path_rev)
+    hash_pound = "" if "@" in str(rev) else "#"
+    return f"{path}{hash_pound}{rev}"
+
+#-------------------------------------------------------------------------------
+# 'p4 diff2' can be very slow. We'll use Python's built in difflib
+class _Difflib(object):
+    def __init__(self, left_path_rev, right_path_rev):
+        class P4PrintText(object):
+            def __init__(self):   self.data = b""
+            def on_text(self, t): self.data += t
+
+        left_print = P4PrintText()
+        on_text = left_print.on_text
+        for x in P4.print(_as_p4_path(left_path_rev)).read(on_text=on_text):
+            pass
+
+        right_print = P4PrintText()
+        on_text = right_print.on_text
+        for x in P4.print(_as_p4_path(right_path_rev)).read(on_text=on_text):
+            pass
+
+        import io
+        import difflib
+        diff = difflib.unified_diff(
+            list(io.StringIO(left_print.data.decode(errors="ignore"))),
+            list(io.StringIO(right_print.data.decode(errors="ignore"))),
+            fromfile="from",
+            tofile="to"
+        )
+        self._diff_iter = iter(diff)
+        self.stdout = self
+
+    def readline(self, *args):
+        ret = next(self._diff_iter, None)
+        if ret:
+            return ret.encode()
+        return b""
+
+    def wait(self): pass
+    def close(self): pass
+
+#-------------------------------------------------------------------------------
+# 'p4 diff2' can be very slow, we can use 'git diff' and 'p4 print'
+class _Diff2(object):
+    _git_bin = "git"
+
+    def __init__(self, left_path_rev, right_path_rev):
+        self._left_path = None
+        self._right_path = None
+
+        # Fetch the left file
+        left_fd, left_path = tempfile.mkstemp()
+        on_text = lambda x: os.write(left_fd, x)
+        for x in P4.print(_as_p4_path(left_path_rev)).read(on_text=on_text):
+            pass
+        os.close(left_fd)
+        self._left_path = left_path
+
+        # Fetch the right file
+        right_fd, right_path = tempfile.mkstemp()
+        on_text = lambda x: os.write(right_fd, x)
+        for x in P4.print(_as_p4_path(right_path_rev)).read(on_text=on_text):
+            pass
+        os.close(right_fd)
+        self._right_path = right_path
+
+        # Spin up git diff
+        diff = (
+            _Diff2._git_bin,
+            "diff",
+            "--minimal",
+            "--unified=1",
+            "--diff-algorithm=histogram", # myers, patience
+            left_path,
+            right_path
+        )
+        proc = sp.Popen(diff, stdout=sp.PIPE, stderr=sp.DEVNULL)
+        self.stdout = proc.stdout
+        self.wait = proc.wait
+
+    def __del__(self):
+        def do_unlink(path):
+            try: os.unlink(path)
+            except (OSError, FileNotFoundError): pass
+
+        if self._right_path: do_unlink(self._right_path)
+        if self._left_path: do_unlink(self._left_path)
+
+#-------------------------------------------------------------------------------
+def _get_diff_hunk_for_line(right_line_no, left_path_rev, right_path_rev=None):
+    if right_path_rev:
+        proc = _Difflib(left_path_rev, right_path_rev)
+    else:
+        cmd = ("p4", "diff", "-Od", "-du", _as_p4_path(left_path_rev))
+        proc = sp.Popen(cmd, stdout=sp.PIPE, stderr=sp.DEVNULL)
+
+    def inner_iter():
+        for line in line_iter:
+            if line.startswith(b"@@"):
+                break
+            yield line
+
+        proc.stdout.close()
+        proc.wait()
+
+    offset = 0
+    line_iter = iter(proc.stdout.readline, b"")
+    for line in (x for x in line_iter if x.startswith(b"@@")):
+        m = re.match(rb"@@\s+-(\d+),(\d+)\s+\+(\d+),(\d+)", line)
+        if not m:
+            continue
+        left_no, left_n, right_no, right_n = map(int, m.groups())
+
+        if right_line_no >= right_no and right_line_no <= right_no + right_n:
+            return inner_iter(), (left_no, right_no)
+
+        if right_no > right_line_no:
+            break
+
+        offset += left_n - right_n
+
+    proc.stdout.close()
+    proc.wait()
+    return None, right_line_no + offset
+
+#-------------------------------------------------------------------------------
+def _remap_line_no(right_line_no, left_path_rev, right_path_rev=None):
+    line_iter, param = _get_diff_hunk_for_line(right_line_no, left_path_rev, right_path_rev)
+    if not line_iter:
+        return param
+
+    left_no, right_no = param
+    for line in line_iter:
+        if line[0] == ord(b"-"):
+            left_no += 1
+            continue
+
+        if right_no == right_line_no:
+            break
+
+        if line[0] == ord(b" "):
+            left_no += 1
+            right_no += 1
+        elif line[0] == ord(b"+"):
+            right_no += 1
+
+    return left_no
+
+#-------------------------------------------------------------------------------
+def _get_edit_rev(path_rev, line_no):
+    annotate_iter = P4.annotate(_as_p4_path(path_rev), dwl=True, q=True).read()
+    header = next(annotate_iter)
+    try: lines = list(next(annotate_iter) for x in range(line_no))
+    except RuntimeError: raise IndexError
+    return int(lines[-1].lower)
+
+#-------------------------------------------------------------------------------
+def _select_line(path_rev):
+    proc = sp.Popen(("p4", "print", "-q", _as_p4_path(path_rev)), stdout=sp.PIPE)
+
+    def line_iter():
+        for i, line in enumerate(proc.stdout):
+            yield "%4d %s" % (i + 1, line.strip().decode())
+
+    line_no = 0
+    for line in fzf.run(line_iter()):
+        if m := re.match(r"\s*(\d+)", line):
+            line_no = int(m.group(1))
+            break
+
+    proc.wait()
+    proc.stdout.close()
+    return line_no
+
+
+
+#-------------------------------------------------------------------------------
+class Who(flow.cmd.Cmd):
+    """ Find the changelist that modified a particular line of a file. Who will
+    do its best to follow integrations and moves to find the changelist that edited
+    the line. If no line number is given then an prompt will be given to select
+    the line to follow. If 'path' has a '#rev' or '@cl' suffix then no remapping
+    against local edits is skipped."""
+    path      = flow.cmd.Arg(str, "Path to the file (P4 paths & #rev/@cl suffix allowed)")
+    line      = flow.cmd.Arg(0, "Line number to check modifications of")
+    printdiff = flow.cmd.Opt(False, "Print a unified diff of the line change - no P4V")
+    noremap   = flow.cmd.Opt(False, "Skip remapping line numbers against local edits")
+
+    def main(self):
+        # Get the input file, we'll assume it has a good Perforce environment so
+        # we'll chdir alongside it. Might be edited so diff it to adjust line
+        self.print_info("Parsing input")
+        path = self.args.path
+
+        if "#" in path:
+            path_rev = path.split("#", 1)
+            remap_local = False
+        elif "@" in path:
+            path, rev = path.split("@", 1)
+            path_rev = (path, "@" + rev)
+            remap_local = False
+        else:
+            path_rev = (path, "head")
+            remap_local = True
+        path = path_rev[0]
+
+        # Try and pick up a suitable Perforce environment from the on-disk file
+        if os.path.isfile(path):
+            path_rev = (path, "have")
+            path_dir, path = os.path.split(path)
+            if path_dir:
+                os.chdir(path_dir)
+
+        # Path needs converting to a depot path
+        if os.path.isfile(path) or not path.startswith("//"):
+            path_rev = (P4.where(path).depotFile, path_rev[1])
+        print("Depot path:", path_rev[0])
+
+        # Is the file a binary file?
+        is_binary = False
+        try:
+            fstat = P4.fstat(path_rev[0] + "#head", m=1, T="headType").run()
+            is_binary = "binary" in fstat.headType
+        except P4.Error:
+            pass
+
+        if is_binary:
+            raise ValueError("Unable to follow line-by-line edits of binary files")
+
+        line_no = self.args.line
+        if line_no <= 0:
+            if self.is_interactive():
+                print("Select a line from file;", end="")
+                line_no = _select_line(path_rev)
+            if line_no <= 0:
+                self.print_error("No line was selected")
+                return False
+            print("Using line", line_no)
+        elif remap_local and not self.args.noremap:
+            for _ in P4.opened(path_rev[0]):
+                self.print_warning(path_rev[0], "is open for edit")
+                print("Attempting to remap line number against local edits")
+                line_no = _remap_line_no(line_no, path_rev)
+                print("Line", line_no, "perhaps")
+                break
+
+        # Use annotate to find out how a line changed, going further if the
+        # change it was an integration or branch
+        self.print_info("Following edits", flow.cmd.text.grey("... (a)nnotate (f)ilelog (d)iff"))
+        who_cl = 0
+        while not who_cl:
+            print("   ", _as_p4_path(path_rev), end="\r")
+            def print_step(x):
+                print(flow.cmd.text.grey(x), end="")
+
+            # Who last edited out target line?
+            print_step("a")
+            try:
+                edit_rev = _get_edit_rev(path_rev, line_no)
+                edit_path_rev = (path_rev[0], edit_rev)
+            except IndexError:
+                raise IndexError(f"Line {line_no} does not seem to exist")
+
+            # Adjust the edit path if it came from somewhere else
+            print_step("f")
+            log = P4.filelog(_as_p4_path(edit_path_rev), s=True, m=1).run()
+            action = getattr(log, "action0")
+            if action not in ("branch", "integrate", "edit", "add", "move/add"):
+                assert False, "Unknown action; " + action + str(log.as_dict())
+
+            if action not in ("edit", "add"):
+                try:
+                    for how, path, rev in zip(getattr(log, "how0,"), getattr(log, "file0,"), getattr(log, "erev0,")):
+                        if "from" in how:
+                            edit_path_rev = (path, int(rev[1:]))
+                            break
+                    else:
+                        assert False, "no 'from' filelog.how found"
+                except AttributeError:
+                    # This is probably as far as we can go.
+                    who_cl = getattr(log, "change0")
+            else:
+                who_cl = getattr(log, "change0")
+
+            # Remap the line number from where we are (path-rev) to the edit
+            print_step("d")
+            edit_line_no = _remap_line_no(line_no, edit_path_rev, path_rev)
+            if edit_line_no != line_no:
+                print("\n   ", line_no, "->", edit_line_no, end="")
+
+            path_rev = edit_path_rev
+            line_no = edit_line_no
+            print()
+
+        # Inform the user about the change
+        print()
+        self.print_info("Attribution")
+        white = flow.cmd.text.white
+        print(white("Where:     "), _as_p4_path(edit_path_rev))
+        print(white("Line:      "), edit_line_no)
+        desc = P4.describe(who_cl).run()
+        print(white("Changelist:"), desc.change)
+        print(white("User:      "), desc.user)
+        print(white("Time:      "), time.ctime(int(desc.time)))
+        print(white("Description:"))
+        print(flow.cmd.text.light_yellow(desc.desc.rstrip()), sep="")
+
+        # Open the change in a gui.
+        if not self.args.printdiff and self.is_interactive():
+            try:
+                print("Starting pv4c...", end="")
+                p4utils.run_p4vc("change", desc.change)
+                print()
+                return
+            except:
+                print()
+                self.print_warning("Unabled to start 'p4vc'")
+
+        # Show a unfied diff of the line edit
+        self.print_info("Diff")
+
+        # Is the change an add action?
+        if edit_path_rev[1] == 1:
+            print("(No diff to show. This is revision 1)")
+            return
+
+        prev_path_rev = (edit_path_rev[0], edit_path_rev[1] - 1)
+        hunk_iter, param = _get_diff_hunk_for_line(line_no, prev_path_rev, edit_path_rev)
+        if not hunk_iter:
+            return
+
+        max_width = shutil.get_terminal_size(fallback=(9999,0))[0] - 2
+        for line in hunk_iter:
+            prefix = line[:1].decode()
+            line = line[1:].rstrip().expandtabs(4).decode()
+
+            decorator = lambda x: x
+            if prefix == "+":   decorator = flow.cmd.text.green
+            elif prefix == "-": decorator = flow.cmd.text.red
+            print(decorator(prefix + line[:max_width]))
diff --git a/Engine/Extras/ushell/channels/unreal/perforce/cmds/workspace.py b/Engine/Extras/ushell/channels/unreal/perforce/cmds/workspace.py
new file mode 100644
--- /dev/null
+++ b/Engine/Extras/ushell/channels/unreal/perforce/cmds/workspace.py
@@ -0,0 +1,180 @@
+# Copyright Epic Games, Inc. All Rights Reserved.
+
+import os
+import fzf
+import socket
+import p4utils
+import flow.cmd
+from peafour import P4
+from pathlib import Path, PurePosixPath
+
+#-------------------------------------------------------------------------------
+def _build_depot_path():
+    ret = "//"
+    for item in fzf.run((x.name for x in P4.depots()), prompt=ret):
+        ret += item
+        break
+    else:
+        return ""
+
+    ret = PurePosixPath(ret)
+
+    ok_but = "[ok]"
+    while True:
+        def read_dirs(p):
+            yield ok_but
+            yield ".."
+            yield from (x.dir for x in P4.dirs(str(p) + "/*"))
+        for item in fzf.run(read_dirs(ret), prompt=str(ret) + "/"):
+            if item == ok_but:
+                return ret
+            if item == "..":
+                ret = ret.parent
+                continue
+            ret /= item
+            break
+        else:
+            return ret
+
+#-------------------------------------------------------------------------------
+def _check_dir_is_empty(dir):
+    for item in dir.glob("*"):
+        raise ValueError(f"Directory '{dir}' is not empty")
+
+#-------------------------------------------------------------------------------
+def _populate_ue_branch(*args):
+    _populate_ue_branch_engine(*args)
+    _populate_ue_branch_project(*args)
+
+#-------------------------------------------------------------------------------
+def _populate_ue_branch_engine(depot_path, dry_run):
+    minimal_sync_specs = (
+        depot_path / "Engine/Build/Build.version",
+        depot_path / "Engine/Source/*.Target.cs",
+        depot_path / "*",
+    )
+    if dry_run:
+        print(*minimal_sync_specs, sep="\n")
+        return
+
+    for item in P4.sync(*minimal_sync_specs):
+        print(item.depotFile)
+
+#-------------------------------------------------------------------------------
+def _populate_ue_branch_project(depot_path, dry_run):
+    """
+    class ProjectRoots(object):
+        def __init__(self): self._root_dirs = []
+        def __iter__(self): return iter(self._root_dirs)
+        def on_text(self, data):
+            import io
+            for line in io.StringIO(data.decode()):
+                line = line.strip()
+                if not line.startswith(";"):
+                    self._root_dirs.append(line)
+
+    project_roots = ProjectRoots()
+    prnt = P4.print(depot_path / "*.uprojectdirs", q=True)
+    for x in prnt.read(on_text=project_roots.on_text):
+        pass
+
+    uproject_specs = (depot_path / x / "*/*.uproject" for x in project_roots)
+    for item in P4.files(uproject_specs).read(on_error=False):
+        uproject_path = PurePosixPath(item.depotFile)
+        # ignore Templates/ ??|?
+        print(uproject_path.stem, uproject_path.parent)
+    """
+    pass
+
+
+
+#-------------------------------------------------------------------------------
+class Workspace(flow.cmd.Cmd):
+    """ Creates a new workspace """
+    localdir  = flow.cmd.Arg(str, "Directory to create the new workspace in")
+    depotpath = flow.cmd.Arg("", "Depot path to map the workspace to")
+    name      = flow.cmd.Opt("", "The name of the workspace to use")
+    dryrun    = flow.cmd.Opt(False, "Do stuff but don't do it for real")
+
+    def main(self):
+        # Misc details
+        username = p4utils.login()
+        host_name = socket.gethostname()
+
+        # Local FS directory
+        client_dir = Path(self.args.localdir).resolve()
+        if client_dir.is_dir():
+            _check_dir_is_empty(client_dir)
+        elif not self.args.dryrun:
+            client_dir.mkdir(parents=True)
+        os.chdir(client_dir)
+
+        self.print_info("Environment")
+        print("User:", username)
+        print("Host:", host_name)
+        print("Destination:", client_dir)
+
+        # Depot path
+        depot_path = self.args.depotpath
+        if not depot_path:
+            print("Building workspace depot path;", end="")
+            depot_path = _build_depot_path()
+            print()
+        else:
+            depot_path = depot_path.replace("\\", "/")
+            depot_path = PurePosixPath(depot_path)
+
+        if depot_path.root != "//" or len(depot_path.parts) < 2:
+            raise ValueError("Depot path is expected to be in a //foo/bar format")
+
+        is_stream = P4.streams(depot_path).run(on_error=False) is not None
+
+        print("Depot path:", depot_path)
+        print("Stream?:", "Yes" if is_stream else "No")
+
+        # Generate a new client name and check it is availble
+        self.print_info("Creating client")
+        client_name = self.args.name or "_".join((
+            username,
+            host_name,
+            client_dir.stem,
+        ))
+        if P4.client(client_name, t=client_name).run(on_error=False):
+            raise ValueError(f"Workspace '{client_name}' already exists")
+        print("Name:", client_name)
+
+        # Create the workspace
+        client_spec = {
+            "Client"      : client_name,
+            "Owner"       : username,
+            "Host"        : host_name,
+            "Options"     : "noallwrite clobber nocompress unlocked nomodtime rmdir",
+            "Root"        : client_dir,
+            "Description" : "Created with ushell's '.p4 workspace'"
+        }
+        if is_stream:
+            client_spec["Stream"] = depot_path
+        else:
+            client_spec["View0"] = f"{depot_path}/... //{client_name}/..."
+
+        if not self.args.dryrun:
+            print("Creating...", end="")
+            P4.client(i=True).run(input_data=client_spec)
+            print(" done")
+
+        # Write a P4CONFIG file
+        if not self.args.dryrun:
+            p4config_path = client_dir / p4utils.get_p4config_name()
+            with p4config_path.open("wt") as out:
+                print("P4CLIENT=", client_name, sep="", file=out)
+                print("P4USER=", username)
+            print("P4CONFIG:", p4config_path)
+
+        # Is this an Unreal branch?
+        self.print_info("Branch type")
+        ue_depot_root = p4utils.get_branch_root(depot_path)
+        if not str(depot_path)[len(ue_depot_root):]:
+            print("Performing very minimal sync")
+            _populate_ue_branch(depot_path, self.args.dryrun)
+        else:
+            print("New clientspec does not appear to be an Unreal Engine branch")
diff --git a/Engine/Extras/ushell/channels/unreal/perforce/describe.flow.py b/Engine/Extras/ushell/channels/unreal/perforce/describe.flow.py
new file mode 100644
--- /dev/null
+++ b/Engine/Extras/ushell/channels/unreal/perforce/describe.flow.py
@@ -0,0 +1,86 @@
+# Copyright Epic Games, Inc. All Rights Reserved.
+
+import flow.describe
+
+#-------------------------------------------------------------------------------
+bisect = flow.describe.Command()
+bisect.source("cmds/bisect.py", "Bisect")
+bisect.invoke("p4", "bisect")
+
+#-------------------------------------------------------------------------------
+cherrypick = flow.describe.Command()
+cherrypick.source("cmds/cherrypick.py", "Cherrypick")
+cherrypick.invoke("p4", "cherrypick")
+
+#-------------------------------------------------------------------------------
+clean = flow.describe.Command()
+clean.source("cmds/clean.py", "Clean")
+clean.invoke("p4", "clean")
+
+#-------------------------------------------------------------------------------
+client = flow.describe.Command()
+client.source("cmds/workspace.py", "Workspace")
+client.invoke("p4", "workspace")
+
+#-------------------------------------------------------------------------------
+gui = flow.describe.Command()
+gui.source("cmds/gui.py", "Gui")
+gui.invoke("p4", "v")
+
+#-------------------------------------------------------------------------------
+mergedown = flow.describe.Command()
+mergedown.source("cmds/mergedown.py", "MergeDown")
+mergedown.invoke("p4", "mergedown")
+
+#-------------------------------------------------------------------------------
+project = flow.describe.Command()
+project.source("cmds/project.py", "Change")
+project.invoke("project")
+
+#-------------------------------------------------------------------------------
+reset = flow.describe.Command()
+reset.source("cmds/clean.py", "Reset")
+reset.invoke("p4", "reset")
+
+#-------------------------------------------------------------------------------
+switch = flow.describe.Command()
+switch.source("cmds/switch.py", "Switch")
+switch.invoke("p4", "switch")
+
+#-------------------------------------------------------------------------------
+switch_list = flow.describe.Command()
+switch_list.source("cmds/switch.py", "List")
+switch_list.invoke("p4", "switch", "list")
+
+#-------------------------------------------------------------------------------
+sync = flow.describe.Command()
+sync.source("cmds/sync.py", "Sync")
+sync.invoke("p4", "sync")
+
+sync_edit = flow.describe.Command()
+sync_edit.source("cmds/sync.py", "Edit")
+sync_edit.invoke("p4", "sync", "edit")
+
+#-------------------------------------------------------------------------------
+who = flow.describe.Command()
+who.source("cmds/who.py", "Who")
+who.invoke("p4", "who")
+
+#-------------------------------------------------------------------------------
+boot = flow.describe.Command()
+boot.source("boot.py", "Boot")
+boot.invoke("boot")
+boot.prefix("$")
+
+#-------------------------------------------------------------------------------
+tips = flow.describe.Command()
+tips.source("tips.py", "Tips")
+tips.invoke("tip")
+tips.prefix("$")
+
+
+
+#-------------------------------------------------------------------------------
+channel = flow.describe.Channel()
+channel.version("0")
+channel.parent("unreal.core")
diff --git a/Engine/Extras/ushell/channels/unreal/perforce/pylib/p4utils.py b/Engine/Extras/ushell/channels/unreal/perforce/pylib/p4utils.py
new file mode 100644
--- /dev/null
+++ b/Engine/Extras/ushell/channels/unreal/perforce/pylib/p4utils.py
@@ -0,0 +1,455 @@
+# Copyright Epic Games, Inc. All Rights Reserved.
+
+import os
+import re
+import time
+import flow.cmd
+import threading
+from peafour import P4
+import subprocess as sp
+from pathlib import Path
+
+#-------------------------------------------------------------------------------
+class _ViewWildcard(object):
+    def __init__(self, view, *, allow_empty_wildcards=False, case_sensitive=False):
+        view = view.replace("\\", "/")
+
+        if view.endswith("/"):
+            view += "\x01"
+        elif view.endswith("/..."):
+            view = view[:-3] + "\x01"
+
+        rooted = view.startswith("//")
+
+        regex = ""
+        sep = ""
+        for piece in (x for x in view.split("/") if x):
+            regex += sep
+            sep = "/"
+
+            if piece == "..." and allow_empty_wildcards:
+                regex += "(?:[^/]+/)*"
+                sep = ""
+                continue
+
+            piece = piece.replace("...", "\x01")
+            piece = piece.replace("*", "\x02")
+            piece = re.escape(piece)
+            piece = piece.replace("\x01", ".*")
+            piece = piece.replace("\x02", "[^/]*")
+            regex += piece
+
+        regex = "//" + regex if rooted else regex
+        regex = (regex or "^") + "$"
+        flags = None if case_sensitive else re.IGNORECASE
+        self._re = re.compile(regex, flags)
+
+    def test(self, depot_path):
+        return self._re.search(depot_path) is not None
+
+
+
+#-------------------------------------------------------------------------------
+p4_set_result = None
+def get_p4_set(prop_name):
+    if ret := os.getenv(prop_name):
+        return ret
+
+    global p4_set_result
+    if not p4_set_result:
+        p4_set_result = {}
+        try: proc = sp.Popen(("p4", "set", "-q"), stdout=sp.PIPE)
+        except: return
+
+        for line in iter(proc.stdout.readline, b""):
+            try: key, value = line.split(b"=", 1)
+            except ValueError: continue
+            p4_set_result[key.decode()] = value.strip().decode()
+
+        proc.wait()
+        proc.stdout.close()
+
+    return p4_set_result.get(prop_name, None)
+
+#-------------------------------------------------------------------------------
+def login():
+    try:
+        detail = P4.login(s=True).run()
+    except P4.Error:
+        raise EnvironmentError("No valid Perforce session found. Run 'p4 login' to authenticate.")
+    return getattr(detail, "User", None)
+
+#-------------------------------------------------------------------------------
+def get_p4config_name():
+    return os.path.basename(get_p4_set("P4CONFIG") or ".p4config.txt")
+
+#-------------------------------------------------------------------------------
+def has_p4config(start_dir):
+    from pathlib import Path
+    p4config = get_p4config_name()
+    for dir in (Path(start_dir) / "x").parents:
+        candidate = dir / p4config
+        if os.path.isfile(candidate):
+            return True, candidate
+    return False, p4config
+
+#-------------------------------------------------------------------------------
+def create_p4config(p4config_path, client, username, port=None):
+    with open(p4config_path, "wt") as out:
+        print_args = { "sep" : "", "file" : out }
+        print("P4CLIENT=", client, **print_args)
+        print("P4USER=", username, **print_args)
+        if port:
+            print("P4PORT", port, **print_args)
+
+#-------------------------------------------------------------------------------
+def ensure_p4config(start_dir=None):
+    start_dir = start_dir or os.getcwd()
+    found, p4config_name = has_p4config(start_dir)
+    if found:
+        return p4config_name, False
+
+    username = login()
+
+    # Get the client for 'start_dir'
+    client = get_client_from_dir(start_dir, username)
+    if not client:
+        return
+    client, root_dir = client
+
+    # Now we know where to locate a p4config file
+    p4config_path = f"{root_dir}/{p4config_name}"
+    create_p4config(p4config_path, client, username)
+    return p4config_path, True
+
+#-------------------------------------------------------------------------------
+def get_client_from_dir(root_dir, username):
+    import socket
+    host_name = socket.gethostname().lower()
+
+    root_dir = Path(root_dir)
+
+    clients = (x for x in P4.clients(u=username) if x.Host.lower() == host_name)
+    for client in clients:
+        if root_dir.is_relative_to(client.Root):
+            return client.client, client.Root
+
+#-------------------------------------------------------------------------------
+def get_branch_root(depot_path):
+    depot_path = str(depot_path)
+
+    def fstat_paths():
+        limit = 5 # ...two of which are always required
+        query_path = "//"
+        for piece in depot_path[2:].split("/")[:limit]:
+            query_path += piece + "/"
+            yield query_path + "GenerateProjectFiles.bat"
+
+    print("Probing for well-known file:")
+    for x in fstat_paths():
+        print(" ", x)
+
+    fstat = P4.fstat(fstat_paths(), T="depotFile")
+    root_path = fstat.run(on_error=False)
+    if root_path:
+        return "/".join(root_path.depotFile.split("/")[:-1]) + "/"
+
+    raise ValueError("Unable to establish branch root")
+
+#-------------------------------------------------------------------------------
+def run_p4vc(*args):
+    try:
+        proc = sp.run(("p4vc", *args))
+    except FileNotFoundError:
+        proc = sp.run(("p4v", "-p4vc", *args))
+    return proc.returncode
+
+
+
+#-------------------------------------------------------------------------------
+class TempBranchSpec(object):
+    def __init__(self, use, username, from_path, to_path, ignore_streams=False):
+        import hashlib
+        id = hashlib.md5()
+        id.update(from_path.encode())
+        id.update(to_path.encode())
+        id = id.hexdigest()[:6]
+        self._name = f"{username}-ushell.{use}-{id}"
+
+        # To map between streams we need to extract the internal branchspec that
+        # Perforce builds. If from/to aren't related streams it will fail so we
+        # fallback to a conventional trivial branchspec.
+        try:
+            if ignore_streams:
+                raise P4.Error("")
+
+            branch = P4.branch(self._name, o=True, S=from_path[:-1], P=to_path[:-1])
+            result = branch.run()
+            spec = result.as_dict()
+        except P4.Error:
+            spec = {
+                "Branch" : self._name,
+                "View0"  : f"{from_path}... {to_path}...",
+            }
+
+        P4.branch(i=True).run(input_data=spec, on_error=False)
+
+    def __del__(self):
+        P4.branch(self._name, d=True).run()
+
+    def __str__(self):
+        return self._name
+
+
+
+#-------------------------------------------------------------------------------
+def _kb_string(value):
+    return format(value // 1024, ",") + "KB"
+
+#-------------------------------------------------------------------------------
+class _SyncRota(object):
+    class _Worker(object):
+        def __init__(self, id):
+            self.id = id
+            self.work_items = []
+            self.burden = 0
+            self.done_size = 0
+            self.done_items = 0
+            self.error = False
+
+    def __init__(self, changelist, worker_count):
+        # An extra worker is used for syncs to #0 (i.e. deletes). When using the
+        # -L argument a move will be two path#rev rota items. If they are actioned
+        # by the same p4 worker only one item gets synced as "p4 sync" ignores the
+        # other. So a move's two path#rev actions are give to two separate workers.
+        self._workers = [_SyncRota._Worker(x) for x in range(worker_count + 1)]
+        self.changelist = str(changelist)
+
+    def add_work(self, item, rev, cost):
+        if rev == 0:
+            worker = self._workers[0]
+        else:
+            worker = min(self._workers[1:], key=lambda x: x.burden)
+        worker.work_items.append((item, rev, cost))
+        worker.burden += cost
+
+    def sort(self):
+        direction = 1
+        for worker in self._workers:
+            worker.work_items.sort(key=lambda x: x[2] * direction)
+            direction *= -1
+
+    def read_workers(self):
+        yield from (x for x in self._workers if x.work_items)
+
+#-------------------------------------------------------------------------------
+class Syncer(object):
+    def __init__(self):
+        self._paths = []
+        self._excluded_views = set()
+
+    def _read_sync_specs(self, include_excluded=True):
+        cl_suffix = "@" + self._rota.changelist
+        cl_suffix = "#0" if cl_suffix == "@0" else cl_suffix
+        for depot_path in self._paths:
+            yield depot_path + cl_suffix
+
+        if include_excluded:
+            # Using "@0" results in slow queries it seems
+            yield from (x + "#0" for x in self._excluded_views)
+
+    def _is_excluded(self, path):
+        return next((True for x in self._view_excludes if x.test(path)), False)
+
+    def _build_exclude_re(self):
+        ret = [_ViewWildcard(x, allow_empty_wildcards=True) for x in self._excluded_views]
+        self._view_excludes = ret
+
+    def add_path(self, dir):
+        self._paths.append(dir)
+
+    def add_exclude(self, view):
+        self._excluded_views.add(view)
+
+    def schedule(self, changelist, worker_count=8):
+        self._build_exclude_re()
+
+        # P4.<cmd> uses p4's Python-marshalled output (the -G option). However the
+        # "p4 sync -n" will report open files via a "info" message instead of a
+        # structured "stat" one. So we capture those and fake a reply
+        def read_items():
+            misreports = []
+            def on_info(data):
+                if m := re.match(r"\s*(//[^#]+)#(\d+).+is opened and", data.data):
+                    misreports.append((m.group(1), m.group(2)))
+
+            yield from P4.sync(self._read_sync_specs(), n=True).read(on_error=False, on_info=on_info)
+
+            class FakeItem(object): pass
+            for depot_path, rev in misreports:
+                item = FakeItem()
+                item.depotFile = depot_path
+                item.rev = rev
+                item.fileSize = 1 # to save looking this up we'll use a nominal size
+                item.action = "edit"
+                yield item
+
+        self._rota = _SyncRota(changelist, worker_count)
+
+        # Fill the rota
+        total_size = 0
+        count = 0
+        for item in read_items():
+            depot_path = item.depotFile
+            rev = int(item.rev)
+
+            if self._is_excluded(depot_path):
+                if item.action != "deleted":
+                    continue
+                rev = 0
+
+            if count % 17 == 0:
+                print("\r" + str(count), "files", f"({_kb_string(total_size)})", end="")
+
+            size = int(getattr(item, "fileSize", 0)) # deletes have no size attr
+
+            # When switching some files are metadata only changes and no content
+            # is transferred.
+            if item.action == "replaced":
+                size = 0
+
+            # When syncing backwards files need to be removed. Perforce has been
+            # seen reporting these correctly as a delete action but with a revision
+            # that is not 0. So we'll just force a rev=0
+            if item.action == "deleted" and rev >= 1:
+                rev = 0
+
+            self._rota.add_work(depot_path, rev, size)
+
+            total_size += size
+            count += 1
+        self._rota.sort()
+        print("\r" + str(count), "files", f"({_kb_string(total_size)})")
+
+    def sync(self, *, dryrun=False, echo=False):
+        # Sum up what we have to do
+        total_burden = sum(x.burden for x in self._rota.read_workers())
+        total_items = sum(len(x.work_items) for x in self._rota.read_workers())
+        lock = threading.Lock()
+
+        print(f"Fetching {_kb_string(total_burden)} in {total_items} files")
+
+        # Launch the worker threads
+        def sync_thread(worker):
+            def on_error(p4_error):
+                if "not enough space" in p4_error.data:
+                    worker.error = "Out of disk space"
+                    raise EOFError()
+                else:
+                    with lock:
+                        print("sync/error: ", p4_error.data, flush=True)
+
+            try:
+                # The 'p4 sync -L' option reportedly does less work on the server
+                # so that's what we use here. Note that it only works with "#rev"
+                # syntax if one wants to sync a specific revision or changelist.
+                def read_sync_items():
+                    for path, rev, size in worker.work_items:
+                        yield f"{path}#{rev}"
+
+                sync = P4(b=8192).sync(read_sync_items(), n=dryrun, L=True)
+                for item in sync.read(on_error=on_error):
+                    if echo:
+                        with lock:
+                            print("sync/worker:", item.depotFile, "#" + str(item.rev), flush=True)
+
+                    size = int(getattr(item, "fileSize", 0))
+                    if item.action == "replaced":
+                        # meta data update only, no content xfer
+                        size = 0
+
+                    worker.done_size += size + 0.01
+            except EOFError:
+                pass
+
+        def create_thread(worker):
+            thread = threading.Thread(target=sync_thread, args=(worker,))
+            thread.start()
+            return thread
+
+        threads = [create_thread(x) for x in self._rota.read_workers()]
+        print(f"Using {len(threads)} workers")
+
+        # While there are active threads, print detail about their progress
+        total_burden += (0.01 * total_items)
+        while not echo:
+            threads = [x for x in threads if x.is_alive()]
+            if not threads:
+                break
+
+            done_size = sum(x.done_size for x in self._rota.read_workers())
+            progress = ((done_size * 1000) // total_burden) / 10
+            print("\r%5.1f%%" % progress, _kb_string(int(done_size)), end="");
+            time.sleep(0.3)
+        else:
+            for thread in threads:
+                thread.join()
+        print("\r...done              ")
+
+        # Check for errors from the workers
+        for worker in (x for x in self._rota.read_workers() if x.error):
+            print(flow.cmd.text.red("!!" + str(worker.error)))
+            return False
+
+        # Nothing more to do if this is a dry run as the remaining tasks need a
+        # sync to operate on.
+        if dryrun:
+            return True
+
+        # P4.sync() returns 'stat' type events but "p4 sync" will report files
+        # with a complex sync scenario only as unstructured 'info' messages. As the
+        # above won't know about these files we'll do a second sync to catch them.
+        global sync_errors
+        sync_errors = False
+        print("Finalising...", end="")
+        def read_depot_files():
+            def on_error(data):
+                msg = data.data.strip()
+                if "up-to-date" in msg: return
+                if "not in client view" in msg: return
+                if "protected namespace - access denied" in msg: return
+
+                print("\n", flow.cmd.text.red(msg), end="")
+                global sync_errors
+                sync_errors = True
+
+            p4_is_daft = []
+            def on_info(data):
+                msg = data.data
+                if " - must resolve #" in msg: return
+                if " - is opened and not being changed" in msg:
+                    try: p4_is_daft.append(msg.split(" - ")[0])
+                    except IOError: pass
+                    return
+                print("\n", flow.cmd.text.light_yellow(data.data), end="")
+
+            sync = P4.sync(self._read_sync_specs(False), n=True)
+            for item in sync.read(on_error=on_error, on_info=on_info):
+                if not self._is_excluded(item.depotFile):
+                    yield item.depotFile + "#" + item.rev
+
+            yield from iter(p4_is_daft)
+
+        count = 0
+        sync = P4.sync(read_depot_files(), L=True)
+        for item in sync.read(on_error=False):
+            if echo:
+                with lock:
+                    print(f"sync/final: {item.depotFile} #{item.rev}")
+            else:
+                count += 1
+                print("\rFinalising...", count, end="")
+
+        print()
+
+        return not sync_errors
diff --git a/Engine/Extras/ushell/channels/unreal/perforce/pylib/peafour.py b/Engine/Extras/ushell/channels/unreal/perforce/pylib/peafour.py
new file mode 100644
--- /dev/null
+++ b/Engine/Extras/ushell/channels/unreal/perforce/pylib/peafour.py
@@ -0,0 +1,193 @@
+# Copyright Epic Games, Inc. All Rights Reserved.
+
+import types
+import marshal
+import threading
+import subprocess as sp
+from pathlib import PurePath
+
+#-------------------------------------------------------------------------------
+class _P4Result(object):
+    def __init__(self, result):
+        super().__setattr__("_result", result)
+
+    def __str__(self):                 return str(self._result)
+    def __contains__(self, key):       return (key in self._result) or (key + "0" in self._result)
+    def as_dict(self):                 return {k.decode():v.decode() for k,v in self._result.items()}
+    def __setattr__(self, key, value): self._result[key.encode()] = str(value).encode()
+
+    def __getattr__(self, key):
+        if (key + "0").encode() in self._result:
+            def as_list():
+                index = 0;
+                while (key + str(index)).encode() in self._result:
+                    indexed_key = (key + str(index)).encode()
+                    yield self._result[indexed_key].decode()
+                    index += 1
+            return as_list()
+
+        ret = self._result.get(str.encode(key))
+        if ret == None: raise AttributeError(key)
+        return ret.decode(errors="replace")
+
+#-------------------------------------------------------------------------------
+class _P4Command(object):
+    @staticmethod
+    def _read_args(*args, **kwargs):
+        for k,v in kwargs.items():
+            if isinstance(v, bool):
+                if v:
+                    yield "-" + k
+            elif v != None:
+                yield f"-{k}={v}" if len(k) > 1 else f"-{k}{v}"
+
+        for arg in args:
+            if isinstance(arg, str) or isinstance(arg, PurePath):
+                yield str(arg)
+
+    def __init__(self, **options):
+        opt_iter = _P4Command._read_args(**options)
+        self._command = ["p4", "-Qutf8", "-G"]
+        self._command += (x for x in opt_iter)
+
+    def start(self, command, *args, **kwargs):
+        self._stdin_args = []
+        for arg in args:
+            if isinstance(arg, str) or isinstance(arg, PurePath):
+                continue
+            if not hasattr(arg, "__iter__"):
+                raise TypeError("P4 arguments can be only strings or sequences")
+            self._stdin_args.append(arg)
+
+        self._proc = None
+
+        arg_iter = _P4Command._read_args(*args, **kwargs)
+
+        if self._stdin_args:
+            self._command.append("-x-")
+        self._command.append(command)
+        self._command += (x for x in arg_iter)
+
+    def __del__(self):              self._close_proc()
+    def __str__(self):              return " ".join(self._command)
+    def __iter__(self):             yield from self._iter()
+    def __getattr__(self, name):    return getattr(self.run(), name)
+
+    def run(self, **kwargs):
+        return next(self._iter(**kwargs), None)
+
+    def read(self, **kwargs):
+        yield from self._iter(**kwargs)
+
+    def _close_proc(self):
+        if self._proc:
+            if self._proc.stdin:
+                self._proc.stdin.close()
+            self._proc.stdout.close()
+            self._proc = None
+
+    def _iter(self, input_data=None, on_error=True, on_info=None, on_text=None):
+        stdin = None
+        if input_data != None:
+            if self._stdin_args:
+                raise _P4.Error("It is unsupported to have both generator-type arguments and input data")
+
+            if isinstance(input_data, dict):
+                input_data = {str(k).encode():str(v).encode() for k,v in input_data.items()}
+            else:
+                raise _P4.Error("Unsupported input data type; " + type(input_data).__name__)
+            stdin = sp.PIPE
+
+        if self._stdin_args:
+            stdin = sp.PIPE
+
+        proc = sp.Popen(self._command, stdout=sp.PIPE, stdin=stdin)
+        self._proc = proc
+
+        if stdin:
+            def stdin_thread_entry():
+                try:
+                    if input_data:
+                        marshal.dump(input_data, proc.stdin, 0)
+                    for args in self._stdin_args:
+                        for arg in args:
+                            arg = str(arg).encode() + b"\n"
+                            proc.stdin.write(arg)
+                except (BrokenPipeError, OSError):
+                    pass
+                finally:
+                    try: proc.stdin.close()
+                    except: pass
+
+            stdin_thread = threading.Thread(target=stdin_thread_entry)
+            stdin_thread.start()
+
+        while True:
+            try: result = marshal.load(proc.stdout)
+            except EOFError: break
+
+            code = result[b"code"]
+            del result[b"code"]
+
+            if code == b"error":
+                if isinstance(on_error, bool):
+                    if on_error:
+                        raise _P4.Error(result[b"data"].decode()[:-1])
+                    continue
+
+                try:
+                    on_error(_P4Result(result))
+                except:
+                    proc.terminate()
+                    raise
+                continue
+
+            if code == b"stat":
+                yield _P4Result(result)
+                continue
+
+            if code == b"text" and on_text:
+                try:
+                    data = result.get(b"data", b"")
+                    on_text(data)
+                    continue
+                except:
+                    proc.terminate()
+                    raise
+
+            if code == b"info" and on_info:
+                try:
+                    on_info(_P4Result(result))
+                    continue
+                except:
+                    proc.terminate()
+                    raise
+
+        if stdin:
+            stdin_thread.join()
+        self._close_proc()
+
+#-------------------------------------------------------------------------------
+class _P4(object):
+    class Error(Exception):
+        def __init__(self, msg):
+            super().__init__("Perforce: " + msg)
+
+    def __init__(self, **options):
+        self._options = options
+
+    def __getattr__(self, command):
+        if command == "Error":
+            return _P4.Error
+
+        def inner(*args, **kwargs):
+            instance = _P4Command(**self._options)
+            instance.start(command, *args, **kwargs)
+            return instance
+
+        return inner
+
+    def __call__(self, **options):
+        return _P4(**options)
+
+P4 = _P4()
diff --git a/Engine/Extras/ushell/channels/unreal/perforce/tips.py b/Engine/Extras/ushell/channels/unreal/perforce/tips.py
new file mode 100644
--- /dev/null
+++ b/Engine/Extras/ushell/channels/unreal/perforce/tips.py
@@ -0,0 +1,12 @@
+# Copyright Epic Games, Inc. All Rights Reserved.
+
+#-------------------------------------------------------------------------------
+class Tips(object):
+    def get_tips(self):
+        return (
+            "'.p4 mergedown' will sort unresolve files into a separate changelist",
+            "Integrate/unshelve changelists between branches with '.p4 cherrypick'",
+            "Filter your sync with a .p4sync.txt file. '.p4 sync edit' to quick edit",
+            "Clean branches with '.p4 clean'",
+            "'p4 clean --dryrun' shows how much disk space would be reclaimed",
+        )
diff --git a/Engine/Extras/ushell/ushell.bat b/Engine/Extras/ushell/ushell.bat
new file mode 100755
--- /dev/null
+++ b/Engine/Extras/ushell/ushell.bat
@@ -0,0 +1,51 @@
+@echo off
+
+:: Here be dragons. This file detects how the user has launched ushell.bat by
+:: either 1) a shortcut or directly in Explorer, or 2) as part of a script. This
+:: is done by examining the cmd.exe command line. If (1) is detected then a new
+:: interactive cmd.exe process is started. Otherwise the environment is set up
+:: in the current host.
+
+setlocal
+set _self=%~f0
+
+:: quotes and separators really upset cmd.exe's interpreter so we remove them
+set _cmdline=%cmdcmdline%
+set _cmdline="%_cmdline:"=%"
+set _cmdline=%_cmdline:&=#%
+set _cmdline=%_cmdline:|=#%
+call:main %_cmdline:~1,-1%
+
+set _boot_bat="%~dp0\channels\flow\nt\boot.bat"
+if "%_breakout%"=="" endlocal & %_boot_bat% %*
+endlocal & start /i cmd.exe /d/k "%_boot_bat% %*"
+
+goto:eof
+
+:main
+:: we only consider breaking out if the first cmd.exe argument is /c
+shift
+if not "%~1"=="/c" goto:eof
+
+:: extract the .bat paths script. this reconstructs paths with spaces in
+:bat_path_loop
+    shift
+    if "%1"=="" goto:bat_path_loop_end
+    set _bat_path=%_bat_path% %~1
+    if not "%~x1"==".bat" goto:bat_path_loop
+:bat_path_loop_end
+
+:: we shouldn't think about breaking out if we didn't find ourself
+if not "%_bat_path:~1%"=="%_self%" goto:eof
+
+:: don't breakout if the user appears to have used separators
+:seperators_loop
+    shift
+    if "%1"=="" goto:seperators_loop_end
+    if "%1"=="#" goto:eof
+    if "%1"=="##" goto:eof
+    goto:seperators_loop
+:seperators_loop_end
+
+set _breakout=1
+goto:eof
diff --git a/Engine/Extras/ushell/ushell.sh b/Engine/Extras/ushell/ushell.sh
new file mode 100644
--- /dev/null
+++ b/Engine/Extras/ushell/ushell.sh
@@ -0,0 +1,43 @@
+#!/bin/bash
+
+# Check that this script has been "sourced" into a host shell
+if [ $0 = "$BASH_SOURCE" ]; then
+    echo You must source ushell into your host shell\;
+    echo
+    echo "    source $0"
+    echo
+    exit 1
+fi
+
+# Determine the host shell
+if [ -n "$BASH" ]; then
+    host_shell=bash
+elif [ -n "$ZSH_NAME" ]; then
+    host_shell=zsh
+elif [ -n "$SHELL" ]; then
+    host_shell=$(basename $SHELL)
+fi
+
+if [ -z $host_shell ]; then
+    echo Error: Unable to determine host shell because \$SHELL is unset.
+    echo
+    return
+fi
+
+# Start the session up
+if [ "$WINDIR" ]; then
+    channel=nt
+    cookie=$(cygpath --windows --absolute /tmp/ushell_$$_shell_cookie)
+else
+    channel=posix
+    cookie=/tmp/ushell_$$_shell_cookie
+fi
+
+$(dirname ${BASH_SOURCE:-$0})/channels/flow/$channel/boot.sh --bootarg=$host_shell,$cookie $@
+
+if [ -f $cookie ]; then
+    chmod u+x $cookie
+    source $cookie
+    unlink $cookie
+    unset cookie
+fi
