Updated ushell [26875346]

diff --git a/Engine/Extras/ushell/channels/flow/core/cmds/debug.py b/Engine/Extras/ushell/channels/flow/core/cmds/debug.py
--- a/Engine/Extras/ushell/channels/flow/core/cmds/debug.py
+++ b/Engine/Extras/ushell/channels/flow/core/cmds/debug.py
@@ -68,7 +68,7 @@ class Invalidate(flow.cmd.Cmd):
                 manifest_path.unlink(missing_ok=True)
 
 #-------------------------------------------------------------------------------
-class Nuke(flow.cmd.Cmd):
+class Wipe(flow.cmd.Cmd):
     """ Wipes everything. """
 
     def main(self):
diff --git a/Engine/Extras/ushell/channels/flow/core/describe.flow.py b/Engine/Extras/ushell/channels/flow/core/describe.flow.py
--- a/Engine/Extras/ushell/channels/flow/core/describe.flow.py
+++ b/Engine/Extras/ushell/channels/flow/core/describe.flow.py
@@ -55,10 +55,10 @@ debug_paths.invoke("debug", "paths")
 debug_paths.prefix("$")
 
 #-------------------------------------------------------------------------------
-debug_nuke = flow.describe.Command()
-debug_nuke.source("cmds/debug.py", "Nuke")
-debug_nuke.invoke("debug", "nuke")
-debug_nuke.prefix("$")
+debug_wipe = flow.describe.Command()
+debug_wipe.source("cmds/debug.py", "Wipe")
+debug_wipe.invoke("debug", "wipe")
+debug_wipe.prefix("$")
 
 #-------------------------------------------------------------------------------
 complete = flow.describe.Command()
diff --git a/Engine/Extras/ushell/channels/unreal/core/cmds/sln.py b/Engine/Extras/ushell/channels/unreal/core/cmds/sln.py
--- a/Engine/Extras/ushell/channels/unreal/core/cmds/sln.py
+++ b/Engine/Extras/ushell/channels/unreal/core/cmds/sln.py
@@ -15,7 +15,8 @@ class _Base(object):
             ppn_path = engine.get_dir() / "Intermediate/ProjectFiles/PrimaryProjectName.txt"
             if not ppn_path.is_file():
                 # Old path, needs to be maintained while projects prior to UE5.1 are supported
-                ppn_path = engine.get_dir() / "Intermediate/ProjectFiles/MasterProjectName.txt"
+                ppn_path = engine.get_dir() / "Intermediate/ProjectFiles/M"
+                ppn_path += "asterProjectName.txt"
             with ppn_path.open("rb") as ppn_file:
                 name = ppn_file.read().decode().strip()
         except FileNotFoundError:
diff --git a/Engine/Extras/ushell/channels/unreal/core/pylib/vs/dte.cpp b/Engine/Extras/ushell/channels/unreal/core/pylib/vs/dte.cpp
--- a/Engine/Extras/ushell/channels/unreal/core/pylib/vs/dte.cpp
+++ b/Engine/Extras/ushell/channels/unreal/core/pylib/vs/dte.cpp
@@ -1,3 +1,5 @@
+// Copyright Epic Games, Inc. All Rights Reserved.
+
 #define _HAS_EXCEPTIONS 0
 #define PY_SSIZE_T_CLEAN
 
