Updated ushell [26875482]

diff --git a/Engine/Extras/ushell/channels/unreal/core/cmds/project.py b/Engine/Extras/ushell/channels/unreal/core/cmds/project.py
--- a/Engine/Extras/ushell/channels/unreal/core/cmds/project.py
+++ b/Engine/Extras/ushell/channels/unreal/core/cmds/project.py
@@ -303,7 +303,7 @@ class Change(flow.cmd.Cmd):
                     with dest_path.open("wb") as out:
                         out.write(data)
 
-    def _get_nyckels(self):
+    def _get_nyckels(self, primary_dir):
         pass
         """
         context = unreal.Context(primary_dir)
