Updated ushell [30799199]

diff --git a/Engine/Extras/ushell/channels/flow/core/boot.py b/Engine/Extras/ushell/channels/flow/core/boot.py
--- a/Engine/Extras/ushell/channels/flow/core/boot.py
+++ b/Engine/Extras/ushell/channels/flow/core/boot.py
@@ -10,6 +10,7 @@ class Boot(flow.cmd.Cmd):
     startdir    = flow.cmd.Arg("", "Working directory to start the shell in")
     bootarg     = flow.cmd.Opt("", "(for internal use)")
     cleanprompt = flow.cmd.Opt(False, "Do not modify the host shell's prompt")
+    theme       = flow.cmd.Opt("", "Adds a spot of colour with '.ushell theme'")
 
     def main(self):
         # Clear the screen
@@ -69,6 +70,11 @@ class Boot(flow.cmd.Cmd):
         system = channel.get_system()
         shell_class(system).boot_shell(env, shell_cookie)
 
+        # Apply a theme
+        if self.args.theme:
+            import subprocess
+            subprocess.run(("_ushell", "theme", "apply", self.args.theme))
+
     def get_prompt(self):
         pass
 
diff --git a/Engine/Extras/ushell/channels/flow/core/cmds/complete.py b/Engine/Extras/ushell/channels/flow/core/cmds/complete.py
--- a/Engine/Extras/ushell/channels/flow/core/cmds/complete.py
+++ b/Engine/Extras/ushell/channels/flow/core/cmds/complete.py
@@ -7,7 +7,7 @@ import importlib.util as imputil
 
 #-------------------------------------------------------------------------------
 def _set_title_inner(t, s):
-    print("\x1b]2;", t, " - ushell", s, sep=None, file=sys.stdout)
+    print("\x1b]2;", t, " - ushell", s, sep="", file=sys.stdout)
 
 if os.name == "nt":
     def _set_title(t):
diff --git a/Engine/Extras/ushell/channels/flow/core/cmds/theme.py b/Engine/Extras/ushell/channels/flow/core/cmds/theme.py
new file mode 100644
--- /dev/null
+++ b/Engine/Extras/ushell/channels/flow/core/cmds/theme.py
@@ -0,0 +1,572 @@
+# Copyright Epic Games, Inc. All Rights Reserved.
+
+import os
+import flow.cmd
+
+#-------------------------------------------------------------------------------
+class _Colour(object):
+    def __init__(self, r, g, b):
+        if isinstance(r, int):
+            self._r = _Colour._from_srgb(float(r) / 255.0)
+            self._g = _Colour._from_srgb(float(g) / 255.0)
+            self._b = _Colour._from_srgb(float(b) / 255.0)
+        elif isinstance(r, float):
+            self._r = r
+            self._g = g
+            self._b = b
+
+    @staticmethod
+    def _from_srgb(c):
+        if c <= 0.04045:
+            return c / 12.92
+        return ((c + 0.055) / 1.055) ** 2.4
+
+    @staticmethod
+    def _to_srgb(c):
+        if c <= 0.0031308:
+            return c * 12.92
+        return ((c ** (1.0 / 2.4)) * 1.055) - 0.055
+
+    def __repr__(self):
+        return "<%0.3f,%0.3f,%0.3f>" % (self._r, self._g, self._b)
+        return str(self.get_rgb_ints())
+
+    def __sub__(self, rhs): return _Colour(self._r - rhs._r, self._g - rhs._g, self._b - rhs._b)
+    def __add__(self, rhs): return _Colour(self._r + rhs._r, self._g + rhs._g, self._b + rhs._b)
+    def __mul__(self, rhs): return _Colour(self._r * rhs._r, self._g * rhs._g, self._b * rhs._b)
+
+    def __isub__(self, rhs): self._r -= rhs._r; self._g -= rhs._g; self._b -= rhs._b; return self
+    def __iadd__(self, rhs): self._r += rhs._r; self._g += rhs._g; self._b += rhs._b; return self
+    def __imul__(self, rhs): self._r *= rhs._r; self._g *= rhs._g; self._b *= rhs._b; return self
+
+    def scaled(self, scale):
+        return _Colour(self._r * scale, self._g * scale, self._b * scale)
+
+    def dot(self, rhs):
+        return (rhs._r * self._r) + (rhs._g * self._g) + (rhs._b * self._b)
+
+    def get_luminance(self):
+        return _Colour(0.2126, 0.7152, 0.0722).dot(self)
+
+    def get_contrast(self, rhs):
+        lhs_lum = self.get_luminance()
+        rhs_lum =  rhs.get_luminance()
+        lums = sorted((lhs_lum, rhs_lum))
+        return (lums[1] + 0.05) / (lums[0] + 0.05)
+
+    def get_rgb_floats(self):
+        return (self._r, self._g, self._b)
+
+    def get_rgb_ints(self):
+        return (
+            int(_Colour._to_srgb(self._r) * 255.0),
+            int(_Colour._to_srgb(self._g) * 255.0),
+            int(_Colour._to_srgb(self._b) * 255.0)
+        )
+
+
+
+#-------------------------------------------------------------------------------
+class Scheme(object):
+    def __init__(self, fg, bg, palette):
+        self._fg = fg
+        self._bg = bg
+        self._palette = palette
+
+    def is_light(self):     return self._bg.get_luminance() > 0.3
+    def get_fg(self):       return self._fg
+    def get_bg(self):       return self._bg
+    def get_palette(self):  return self._palette
+
+#-------------------------------------------------------------------------------
+def _make_scheme(fg, bg, palette, /, contrast=-1.0):
+    def unpack(rgb_int):
+        r = (rgb_int >> 16) & 0xff
+        g = (rgb_int >>  8) & 0xff
+        b = (rgb_int >>  0) & 0xff
+        return _Colour(r, g, b)
+
+    palette = list(unpack(x) for x in palette)
+    palette[0] = bg = unpack(bg)
+    palette[7] = fg = unpack(fg)
+    scheme = Scheme(fg, bg, palette)
+
+    base = 0.0 if scheme.is_light() else 1.0
+    base = _Colour(base, base, base)
+
+    if contrast <= 0.0:
+        return scheme
+
+    # This'll no. no one will notice...
+    def fix_contrast(col, target):
+        c = bg.get_contrast(col)
+        if c >= target:
+            return
+
+        for i in range(4):
+            d = base - col
+            d *= _Colour(0.2, 0.2, 0.2)
+            col += d
+
+            #c = bg.get_contrast(col)
+            #if c >= target:
+                #break
+
+    fix_contrast(scheme.get_fg(), contrast)
+    for i in range(1, len(palette)):
+        fix_contrast(palette[i], contrast)
+
+    return scheme
+
+
+
+
+#-------------------------------------------------------------------------------
+swe = [
+    "493", "a9e", "alp", "bay", "ben", "dib", "dip", "elm", "fro", "gab",
+    "gan", "lam", "law", "loo", "loy", "met", "mix", "new", "nit", "off",
+    "ohm", "one", "pet", "ren", "sac", "sad", "sap", "six", "tye", "wan",
+]
+
+def _is_winterm():
+    return os.getenv("WT_SESSION") != None
+
+def _send_osc(swatch, *params):
+    r, g, b = swatch.get_rgb_ints()
+    xparsecolor = f"rgb:{r:02x}/{g:02x}/{b:02x}"
+    params = ";".join(str(x) for x in params)
+    print(f"\x1b]", params, ";", xparsecolor, sep="", end="\x07")
+
+def _set_scheme(scheme):
+    swatches = [x for x in scheme.get_palette()]
+
+    fg = scheme.get_fg()
+    bg = scheme.get_bg()
+    _send_osc(fg, 10)
+    _send_osc(bg, 11)
+    _send_osc(fg, 12)
+    if not _is_winterm():
+        swatches[0] = bg
+        swatches[7] = fg
+
+    for i, swatch in enumerate(swatches):
+        _send_osc(swatch, 4, i)
+
+#-------------------------------------------------------------------------------
+def _showcase_scheme(name, scheme):
+    import random
+    random.shuffle(swe)
+
+    print("\x1b[38;2;", ";".join(str(x) for x in scheme.get_fg().get_rgb_ints()), sep="", end="m")
+    print("\x1b[48;2;", ";".join(str(x) for x in scheme.get_bg().get_rgb_ints()), sep="", end="m")
+    print("%-25s" % name, end="")
+
+    # 0  8  7 15        gradient
+    # 1  2  3  4  5  6  r g rg b rb gb
+    # 9 10 11 12 13 14  light
+    swatch_indices = (8,7,15, 1,9, 2,10, 3,11, 4,12, 5,13, 6,14)
+
+    palette = scheme.get_palette()
+    bg = scheme.get_bg()
+    for i in swatch_indices:
+        swatch = palette[i]
+        rgb = swatch.get_rgb_ints()
+        print("\x1b[38;2;", ";".join(str(x) for x in rgb), sep="", end="m")
+        tag = swe[i] # "%4.1f" % bg.get_contrast(swatch)
+        print(" %x.%s" % (i, tag), end="")
+
+    print("\x1b[0m")
+
+
+
+#-------------------------------------------------------------------------------
+class _ThemeCmd(flow.cmd.Cmd):
+    contrast = flow.cmd.Opt(3.8, "Minimum colour-to-bg contrast to aim for")
+
+    def get_schemes(self):
+        if hasattr(self, "_schemes"):
+            return self._schemes
+
+        global schemes
+
+        self._schemes = {}
+        for name, (fg, bg, palette) in schemes.items():
+            scheme = _make_scheme(fg, bg, palette, contrast=self.args.contrast)
+            self._schemes[name] = scheme
+
+        return self._schemes
+
+    def complete_theme(self, prefix):
+        global schemes
+        yield from iter(schemes)
+
+#-------------------------------------------------------------------------------
+class Set(_ThemeCmd):
+    """ Sets a particular theme by name if supported by the terminal """
+    theme = flow.cmd.Arg(str, "Theme name to apply")
+
+    def main(self):
+        schemes = self.get_schemes()
+        try:
+            scheme = schemes[self.args.theme]
+        except KeyError:
+            self.print_error(f"Unknown theme '{self.args.theme}'")
+            return False
+
+        _set_scheme(scheme)
+
+#-------------------------------------------------------------------------------
+class Random(_ThemeCmd):
+    """ Cycle randomly through all the themes """
+    mode = flow.cmd.Arg("", "Constrain randomness to dark or light themes")
+    one  = flow.cmd.Opt(False, "Pick one theme at random and exit")
+
+    complete_mode = ("dark", "light")
+
+    def main(self):
+        schemes = self.get_schemes()
+        names = list(schemes.keys())
+        what = self.args.mode.lower()
+        if what == "dark":  names = [x for x in names if not schemes[x].is_light()]
+        if what == "light": names = [x for x in names if     schemes[x].is_light()]
+
+        import random
+        random.shuffle(names)
+        random.shuffle(swe)
+
+        print("Press enter to cycle, Ctrl-C to exit\n\n  ", end="")
+        for k in range(8):
+            for i in (k, k + 8):
+                if i == 0:
+                    continue
+                j = i & 7
+                j += 90 if (i >= 8) else 30
+                print(f" \x1b[{j}m", swe[i], sep="", end="")
+        print("\x1b[0m\n")
+
+        while True:
+            for name in names:
+                scheme = schemes[name]
+                _set_scheme(scheme)
+                print(name, "\x1b[0K", end="")
+                if self.args.one:
+                    return
+                try: input("")
+                except KeyboardInterrupt: return
+                print("\x1bM", end="")
+
+#-------------------------------------------------------------------------------
+class Showcase(_ThemeCmd):
+    """ Showcase all the available colour themes """
+
+    def main(self):
+        schemes = self.get_schemes()
+
+        # calculate worst-case contrasts
+        def contrast_key(scheme):
+            bg = scheme.get_bg()
+            min_c = min(bg.get_contrast(x) for x in scheme.get_palette()[1:])
+            return min_c
+
+        def lum_key(scheme):
+            return scheme.get_bg().get_luminance()
+
+        #min_cs = (contrast_key(x[1]) for x in schemes.items())
+        min_cs = (lum_key(x[1]) for x in schemes.items())
+        inner = list(zip(min_cs, schemes.items()))
+        inner = sorted(inner, key=lambda x: x[0])
+        inner = sorted(inner, key=lambda x: x[1][1].is_light())
+
+        for min_c, (name, scheme) in inner:
+            _showcase_scheme(name, scheme)
+
+
+
+#-------------------------------------------------------------------------------
+schemes = {
+    "3024.day"            : (0x4a4543, 0xf7f7f7, (0x090300, 0xdb2d20, 0x01a252, 0xfded02, 0x01a0e4, 0xa16a94, 0xb5e4f4, 0xa5a2a2, 0x5c5855, 0xe8bbd0, 0x3a3432, 0x4a4543, 0x807d7c, 0xd6d5d4, 0xcdab53, 0xf7f7f7)),
+    "3024.night"          : (0xa5a2a2, 0x090300, (0x090300, 0xdb2d20, 0x01a252, 0xfded02, 0x01a0e4, 0xa16a94, 0xb5e4f4, 0xa5a2a2, 0x5c5855, 0xe8bbd0, 0x3a3432, 0x4a4543, 0x807d7c, 0xd6d5d4, 0xcdab53, 0xf7f7f7)),
+    "a9e.dark"            : (0xbbbbbb, 0x000000, (0x000000, 0xbb0000, 0x00bb00, 0xbbbb00, 0x0000bb, 0xbb00bb, 0x00bbbb, 0xbbbbbb, 0x555555, 0xff5555, 0x55ff55, 0xffff55, 0x5555ff, 0xff55ff, 0x55ffff, 0xffffff)),
+    "a9e.light"           : (0x000000, 0xffffff, (0x000000, 0xbb0000, 0x00bb00, 0xbbbb00, 0x0000bb, 0xbb00bb, 0x00bbbb, 0xbbbbbb, 0x555555, 0xff5555, 0x55ff55, 0xffff55, 0x5555ff, 0xff55ff, 0x55ffff, 0xffffff)),
+    "a9e.pastel.dark"     : (0xbbbbbb, 0x000000, (0x4f4f4f, 0xff6c60, 0xa8ff60, 0xffffb6, 0x96cbfe, 0xff73fd, 0xc6c5fe, 0xeeeeee, 0x7c7c7c, 0xffb6b0, 0xceffac, 0xffffcc, 0xb5dcff, 0xff9cfe, 0xdfdffe, 0xffffff)),
+    "a9e.solarized.dark"  : (0x839496, 0x002b36, (0x073642, 0xdc322f, 0x859900, 0xb58900, 0x268bd2, 0xd33682, 0x2aa198, 0xeee8d5, 0x002b36, 0xcb4b16, 0x586e75, 0x657b83, 0x839496, 0x6c71c4, 0x93a1a1, 0xfdf6e3)),
+    "a9e.solarized.light" : (0x657b83, 0xfdf6e3, (0x073642, 0xdc322f, 0x859900, 0xb58900, 0x268bd2, 0xd33682, 0x2aa198, 0xeee8d5, 0x002b36, 0xcb4b16, 0x586e75, 0x657b83, 0x839496, 0x6c71c4, 0x93a1a1, 0xfdf6e3)),
+    "a9e.tango.dark"      : (0xffffff, 0x000000, (0x000000, 0xcc0000, 0x4e9a06, 0xc4a000, 0x3465a4, 0x75507b, 0x06989a, 0xd3d7cf, 0x555753, 0xef2929, 0x8ae234, 0xfce94f, 0x729fcf, 0xad7fa8, 0x34e2e2, 0xeeeeec)),
+    "a9e.tango.light"     : (0x000000, 0xffffff, (0x000000, 0xcc0000, 0x4e9a06, 0xc4a000, 0x3465a4, 0x75507b, 0x06989a, 0xd3d7cf, 0x555753, 0xef2929, 0x8ae234, 0xfce94f, 0x729fcf, 0xad7fa8, 0x34e2e2, 0xeeeeec)),
+    "abernathy"           : (0xeeeeec, 0x111416, (0x000000, 0xcd0000, 0x00cd00, 0xcdcd00, 0x1093f5, 0xcd00cd, 0x00cdcd, 0xfaebd7, 0x404040, 0xff0000, 0x00ff00, 0xffff00, 0x11b5f6, 0xff00ff, 0x00ffff, 0xffffff)),
+    "adventure"           : (0xfeffff, 0x040404, (0x040404, 0xd84a33, 0x5da602, 0xeebb6e, 0x417ab3, 0xe5c499, 0xbdcfe5, 0xdbded8, 0x685656, 0xd76b42, 0x99b52c, 0xffb670, 0x97d7ef, 0xaa7900, 0xbdcfe5, 0xe4d5c7)),
+    "adventuretime"       : (0xf8dcc0, 0x1f1d45, (0x050404, 0xbd0013, 0x4ab118, 0xe7741e, 0x0f4ac6, 0x665993, 0x70a598, 0xf8dcc0, 0x4e7cbf, 0xfc5f5a, 0x9eff6e, 0xefc11a, 0x1997c6, 0x9b5953, 0xc8faf4, 0xf6f5fb)),
+    "afterglow"           : (0xd0d0d0, 0x212121, (0x151515, 0xac4142, 0x7e8e50, 0xe5b567, 0x6c99bb, 0x9f4e85, 0x7dd6cf, 0xd0d0d0, 0x505050, 0xac4142, 0x7e8e50, 0xe5b567, 0x6c99bb, 0x9f4e85, 0x7dd6cf, 0xf5f5f5)),
+    "alabaster"           : (0x000000, 0xf7f7f7, (0x000000, 0xaa3731, 0x448c27, 0xcb9000, 0x325cc0, 0x7a3e9d, 0x0083b2, 0xf7f7f7, 0x777777, 0xf05050, 0x60cb00, 0xffbc5d, 0x007acc, 0xe64ce6, 0x00aacb, 0xf7f7f7)),
+    "alienblood"          : (0x637d75, 0x0f1610, (0x112616, 0x7f2b27, 0x2f7e25, 0x717f24, 0x2f6a7f, 0x47587f, 0x327f77, 0x647d75, 0x3c4812, 0xe08009, 0x18e000, 0xbde000, 0x00aae0, 0x0058e0, 0x00e0c4, 0x73fa91)),
+    "andromeda"           : (0xe5e5e5, 0x262a33, (0x000000, 0xcd3131, 0x05bc79, 0xe5e512, 0x2472c8, 0xbc3fbc, 0x0fa8cd, 0xe5e5e5, 0x666666, 0xcd3131, 0x05bc79, 0xe5e512, 0x2472c8, 0xbc3fbc, 0x0fa8cd, 0xe5e5e5)),
+    "arcoiris"            : (0xeee4d9, 0x201f1e, (0x333333, 0xda2700, 0x12c258, 0xffc656, 0x518bfc, 0xe37bd9, 0x63fad5, 0xbab2b2, 0x777777, 0xffb9b9, 0xe3f6aa, 0xffddaa, 0xb3e8f3, 0xcbbaf9, 0xbcffc7, 0xefefef)),
+    "argonaut"            : (0xfffaf4, 0x0e1019, (0x232323, 0xff000f, 0x8ce10b, 0xffb900, 0x008df8, 0x6d43a6, 0x00d8eb, 0xffffff, 0x444444, 0xff2740, 0xabe15b, 0xffd242, 0x0092ff, 0x9a5feb, 0x67fff0, 0xffffff)),
+    "arthur"              : (0xddeedd, 0x1c1c1c, (0x3d352a, 0xcd5c5c, 0x86af80, 0xe8ae5b, 0x6495ed, 0xdeb887, 0xb0c4de, 0xbbaa99, 0x554444, 0xcc5533, 0x88aa22, 0xffa75d, 0x87ceeb, 0x996600, 0xb0c4de, 0xddccbb)),
+    "ateliersulphurpool"  : (0x979db4, 0x202746, (0x202746, 0xc94922, 0xac9739, 0xc08b30, 0x3d8fd1, 0x6679cc, 0x22a2c9, 0x979db4, 0x6b7394, 0xc76b29, 0x293256, 0x5e6687, 0x898ea4, 0xdfe2f1, 0x9c637a, 0xf5f7ff)),
+    "atom"                : (0xc5c8c6, 0x161719, (0x000000, 0xfd5ff1, 0x87c38a, 0xffd7b1, 0x85befd, 0xb9b6fc, 0x85befd, 0xe0e0e0, 0x000000, 0xfd5ff1, 0x94fa36, 0xf5ffa8, 0x96cbfe, 0xb9b6fc, 0x85befd, 0xe0e0e0)),
+    "atomone.light"       : (0x2a2c33, 0xf9f9f9, (0x000000, 0xde3e35, 0x3f953a, 0xd2b67c, 0x2f5af3, 0x950095, 0x3f953a, 0xbbbbbb, 0x000000, 0xde3e35, 0x3f953a, 0xd2b67c, 0x2f5af3, 0xa00095, 0x3f953a, 0xffffff)),
+    "aurora"              : (0xffca28, 0x23262e, (0x23262e, 0xf0266f, 0x8fd46d, 0xffe66d, 0x0321d7, 0xee5d43, 0x03d6b8, 0xc74ded, 0x292e38, 0xf92672, 0x8fd46d, 0xffe66d, 0x03d6b8, 0xee5d43, 0x03d6b8, 0xc74ded)),
+    "ayu"                 : (0xe6e1cf, 0x0f1419, (0x000000, 0xff3333, 0xb8cc52, 0xe7c547, 0x36a3d9, 0xf07178, 0x95e6cb, 0xffffff, 0x323232, 0xff6565, 0xeafe84, 0xfff779, 0x68d5ff, 0xffa3aa, 0xc7fffd, 0xffffff)),
+    "ayu.light"           : (0x5c6773, 0xfafafa, (0x000000, 0xff3333, 0x86b300, 0xf29718, 0x41a6d9, 0xf07178, 0x4dbf99, 0xffffff, 0x323232, 0xff6565, 0xb8e532, 0xffc94a, 0x73d8ff, 0xffa3aa, 0x7ff1cb, 0xffffff)),
+    "ayu.mirage"          : (0xcbccc6, 0x1f2430, (0x191e2a, 0xed8274, 0xa6cc70, 0xfad07b, 0x6dcbfa, 0xcfbafa, 0x90e1c6, 0xc7c7c7, 0x686868, 0xf28779, 0xbae67e, 0xffd580, 0x73d0ff, 0xd4bfff, 0x95e6cb, 0xffffff)),
+    "banana.blueberry"    : (0xcccccc, 0x191323, (0x17141f, 0xff6b7f, 0x00bd9c, 0xe6c62f, 0x22e8df, 0xdc396a, 0x56b6c2, 0xf1f1f1, 0x495162, 0xfe9ea1, 0x98c379, 0xf9e46b, 0x91fff4, 0xda70d6, 0xbcf3ff, 0xffffff)),
+    "batman"              : (0x6f6f6f, 0x1b1d1e, (0x1b1d1e, 0xe6dc44, 0xc8be46, 0xf4fd22, 0x737174, 0x747271, 0x62605f, 0xc6c5bf, 0x505354, 0xfff78e, 0xfff27d, 0xfeed6c, 0x919495, 0x9a9a9d, 0xa3a3a6, 0xdadbd6)),
+    "belafonte.day"       : (0x45373c, 0xd5ccba, (0x20111b, 0xbe100e, 0x858162, 0xeaa549, 0x426a79, 0x97522c, 0x989a9c, 0x968c83, 0x5e5252, 0xbe100e, 0x858162, 0xeaa549, 0x426a79, 0x97522c, 0x989a9c, 0xd5ccba)),
+    "belafonte.night"     : (0x968c83, 0x20111b, (0x20111b, 0xbe100e, 0x858162, 0xeaa549, 0x426a79, 0x97522c, 0x989a9c, 0x968c83, 0x5e5252, 0xbe100e, 0x858162, 0xeaa549, 0x426a79, 0x97522c, 0x989a9c, 0xd5ccba)),
+    "birdsofparadise"     : (0xe0dbb7, 0x2a1f1d, (0x573d26, 0xbe2d26, 0x6ba18a, 0xe99d2a, 0x5a86ad, 0xac80a6, 0x74a6ad, 0xe0dbb7, 0x9b6c4a, 0xe84627, 0x95d8ba, 0xd0d150, 0xb8d3ed, 0xd19ecb, 0x93cfd7, 0xfff9d5)),
+    "blazer"              : (0xd9e6f2, 0x0d1926, (0x000000, 0xb87a7a, 0x7ab87a, 0xb8b87a, 0x7a7ab8, 0xb87ab8, 0x7ab8b8, 0xd9d9d9, 0x262626, 0xdbbdbd, 0xbddbbd, 0xdbdbbd, 0xbdbddb, 0xdbbddb, 0xbddbdb, 0xffffff)),
+    "blue.matrix"         : (0x00a2ff, 0x101116, (0x101116, 0xff5680, 0x00ff9c, 0xfffc58, 0x00b0ff, 0xd57bff, 0x76c1ff, 0xc7c7c7, 0x686868, 0xff6e67, 0x5ffa68, 0xfffc67, 0x6871ff, 0xd682ec, 0x60fdff, 0xffffff)),
+    "blueberrypie"        : (0xbabab9, 0x1c0c28, (0x0a4c62, 0x99246e, 0x5cb1b3, 0xeab9a8, 0x90a5bd, 0x9d54a7, 0x7e83cc, 0xf0e8d6, 0x201637, 0xc87272, 0x0a6c7e, 0x7a3188, 0x39173d, 0xbc94b7, 0x5e6071, 0x0a6c7e)),
+    "bluedolphin"         : (0xc5f2ff, 0x006984, (0x292d3e, 0xff8288, 0xb4e88d, 0xf4d69f, 0x82aaff, 0xe9c1ff, 0x89ebff, 0xd0d0d0, 0x434758, 0xff8b92, 0xddffa7, 0xffe585, 0x9cc4ff, 0xddb0f6, 0xa3f7ff, 0xffffff)),
+    "bluloco.dark"        : (0xb9c0cb, 0x282c34, (0x41444d, 0xfc2f52, 0x25a45c, 0xff936a, 0x3476ff, 0x7a82da, 0x4483aa, 0xcdd4e0, 0x8f9aae, 0xff6480, 0x3fc56b, 0xf9c859, 0x10b1fe, 0xff78f8, 0x5fb9bc, 0xffffff)),
+    "bluloco.light"       : (0x373a41, 0xf9f9f9, (0x373a41, 0xd52753, 0x23974a, 0xdf631c, 0x275fe4, 0x823ff1, 0x27618d, 0xbabbc2, 0x676a77, 0xff6480, 0x3cbc66, 0xc5a332, 0x0099e1, 0xce33c0, 0x6d93bb, 0xd3d3d3)),
+    "borland"             : (0xffff4e, 0x0000a4, (0x4f4f4f, 0xff6c60, 0xa8ff60, 0xffffb6, 0x96cbfe, 0xff73fd, 0xc6c5fe, 0xeeeeee, 0x7c7c7c, 0xffb6b0, 0xceffac, 0xffffcc, 0xb5dcff, 0xff9cfe, 0xdfdffe, 0xffffff)),
+    "breeze"              : (0xeff0f1, 0x31363b, (0x31363b, 0xed1515, 0x11d116, 0xf67400, 0x1d99f3, 0x9b59b6, 0x1abc9c, 0xeff0f1, 0x7f8c8d, 0xc0392b, 0x1cdc9a, 0xfdbc4b, 0x3daee9, 0x8e44ad, 0x16a085, 0xfcfcfc)),
+    "bright.lights"       : (0xb3c9d7, 0x191919, (0x191919, 0xff355b, 0xb7e876, 0xffc251, 0x76d4ff, 0xba76e7, 0x6cbfb5, 0xc2c8d7, 0x191919, 0xff355b, 0xb7e876, 0xffc251, 0x76d5ff, 0xba76e7, 0x6cbfb5, 0xc2c8d7)),
+    "broadcast"           : (0xe6e1dc, 0x2b2b2b, (0x000000, 0xda4939, 0x519f50, 0xffd24a, 0x6d9cbe, 0xd0d0ff, 0x6e9cbe, 0xffffff, 0x323232, 0xff7b6b, 0x83d182, 0xffff7c, 0x9fcef0, 0xffffff, 0xa0cef0, 0xffffff)),
+    "brogrammer"          : (0xd6dbe5, 0x131313, (0x1f1f1f, 0xf81118, 0x2dc55e, 0xecba0f, 0x2a84d2, 0x4e5ab7, 0x1081d6, 0xd6dbe5, 0xd6dbe5, 0xde352e, 0x1dd361, 0xf3bd09, 0x1081d6, 0x5350b9, 0x0f7ddb, 0xffffff)),
+    "c64"                 : (0x7869c4, 0x40318d, (0x090300, 0x883932, 0x55a049, 0xbfce72, 0x40318d, 0x8b3f96, 0x67b6bd, 0xffffff, 0x000000, 0x883932, 0x55a049, 0xbfce72, 0x40318d, 0x8b3f96, 0x67b6bd, 0xf7f7f7)),
+    "calamity"            : (0xd5ced9, 0x2f2833, (0x2f2833, 0xfc644d, 0xa5f69c, 0xe9d7a5, 0x3b79c7, 0xf92672, 0x74d3de, 0xd5ced9, 0x7e6c88, 0xfc644d, 0xa5f69c, 0xe9d7a5, 0x3b79c7, 0xf92672, 0x74d3de, 0xffffff)),
+    "cga"                 : (0xaaaaaa, 0x000000, (0x000000, 0xaa0000, 0x00aa00, 0xaa5500, 0x0000aa, 0xaa00aa, 0x00aaaa, 0xaaaaaa, 0x555555, 0xff5555, 0x55ff55, 0xffff55, 0x5555ff, 0xff55ff, 0x55ffff, 0xffffff)),
+    "chalk"               : (0xd2d8d9, 0x2b2d2e, (0x7d8b8f, 0xb23a52, 0x789b6a, 0xb9ac4a, 0x2a7fac, 0xbd4f5a, 0x44a799, 0xd2d8d9, 0x888888, 0xf24840, 0x80c470, 0xffeb62, 0x4196ff, 0xfc5275, 0x53cdbd, 0xd2d8d9)),
+    "chalkboard"          : (0xd9e6f2, 0x29262f, (0x000000, 0xc37372, 0x72c373, 0xc2c372, 0x7372c3, 0xc372c2, 0x72c2c3, 0xd9d9d9, 0x323232, 0xdbaaaa, 0xaadbaa, 0xdadbaa, 0xaaaadb, 0xdbaada, 0xaadadb, 0xffffff)),
+    "challengerdeep"      : (0xcbe1e7, 0x1e1c31, (0x141228, 0xff5458, 0x62d196, 0xffb378, 0x65b2ff, 0x906cff, 0x63f2f1, 0xa6b3cc, 0x565575, 0xff8080, 0x95ffa4, 0xffe9aa, 0x91ddff, 0xc991e1, 0xaaffe4, 0xcbe3e7)),
+    "chester"             : (0xffffff, 0x2c3643, (0x080200, 0xfa5e5b, 0x16c98d, 0xffc83f, 0x288ad6, 0xd34590, 0x28ddde, 0xe7e7e7, 0x6f6b68, 0xfa5e5b, 0x16c98d, 0xfeef6d, 0x278ad6, 0xd34590, 0x27dede, 0xffffff)),
+    "ciapre"              : (0xaea47a, 0x191c27, (0x181818, 0x810009, 0x48513b, 0xcc8b3f, 0x576d8c, 0x724d7c, 0x5c4f4b, 0xaea47f, 0x555555, 0xac3835, 0xa6a75d, 0xdcdf7c, 0x3097c6, 0xd33061, 0xf3dbb2, 0xf4f4f4)),
+    "clrs"                : (0x262626, 0xffffff, (0x000000, 0xf8282a, 0x328a5d, 0xfa701d, 0x135cd0, 0x9f00bd, 0x33c3c1, 0xb3b3b3, 0x555753, 0xfb0416, 0x2cc631, 0xfdd727, 0x1670ff, 0xe900b0, 0x3ad5ce, 0xeeeeec)),
+    "cobalt.neon"         : (0x8ff586, 0x142838, (0x142631, 0xff2320, 0x3ba5ff, 0xe9e75c, 0x8ff586, 0x781aa0, 0x8ff586, 0xba46b2, 0xfff688, 0xd4312e, 0x8ff586, 0xe9f06d, 0x3c7dd2, 0x8230a7, 0x6cbc67, 0x8ff586)),
+    "cobalt2"             : (0xffffff, 0x132738, (0x000000, 0xff0000, 0x38de21, 0xffe50a, 0x1460d2, 0xff005d, 0x00bbbb, 0xbbbbbb, 0x555555, 0xf40e17, 0x3bd01d, 0xedc809, 0x5555ff, 0xff55ff, 0x6ae3fa, 0xffffff)),
+    "coffee.theme"        : (0x000000, 0xf5deb3, (0x000000, 0xc91b00, 0x00c200, 0xc7c400, 0x0225c7, 0xca30c7, 0x00c5c7, 0xc7c7c7, 0x686868, 0xff6e67, 0x5ffa68, 0xfffc67, 0x6871ff, 0xff77ff, 0x60fdff, 0xffffff)),
+    "crayonponyfish"      : (0x68525a, 0x150707, (0x2b1b1d, 0x91002b, 0x579524, 0xab311b, 0x8c87b0, 0x692f50, 0xe8a866, 0x68525a, 0x3d2b2e, 0xc5255d, 0x8dff57, 0xc8381d, 0xcfc9ff, 0xfc6cba, 0xffceaf, 0xb0949d)),
+    "cyberdyne"           : (0x00ff92, 0x151144, (0x080808, 0xff8373, 0x00c172, 0xd2a700, 0x0071cf, 0xff90fe, 0x6bffdd, 0xf1f1f1, 0x2e2e2e, 0xffc4be, 0xd6fcba, 0xfffed5, 0xc2e3ff, 0xffb2fe, 0xe6e7fe, 0xffffff)),
+    "cyberpunk"           : (0xe5e5e5, 0x332a57, (0x000000, 0xff7092, 0x00fbac, 0xfffa6a, 0x00bfff, 0xdf95ff, 0x86cbfe, 0xffffff, 0x000000, 0xff8aa4, 0x21f6bc, 0xfff787, 0x1bccfd, 0xe6aefe, 0x99d6fc, 0xffffff)),
+    "dark"                : (0xcccccc, 0x1e1e1e, (0x000000, 0xcd3131, 0x0dbc79, 0xe5e510, 0x2472c8, 0xbc3fbc, 0x11a8cd, 0xe5e5e5, 0x666666, 0xf14c4c, 0x23d18b, 0xf5f543, 0x3b8eea, 0xd670d6, 0x29b8db, 0xe5e5e5)),
+    "dark.pastel"         : (0xffffff, 0x000000, (0x000000, 0xff5555, 0x55ff55, 0xffff55, 0x5555ff, 0xff55ff, 0x55ffff, 0xbbbbbb, 0x555555, 0xff5555, 0x55ff55, 0xffff55, 0x5555ff, 0xff55ff, 0x55ffff, 0xffffff)),
+    "darkermatrix"        : (0x28380d, 0x070c0e, (0x091013, 0x002e18, 0x6fa64c, 0x595900, 0x00cb6b, 0x412a4d, 0x125459, 0x002e19, 0x333333, 0x00381d, 0x90d762, 0xe2e500, 0x00ff87, 0x412a4d, 0x176c73, 0x00381e)),
+    "darkmatrix"          : (0x3e5715, 0x070c0e, (0x091013, 0x006536, 0x6fa64c, 0x7e8000, 0x2c9a84, 0x452d53, 0x114d53, 0x006536, 0x333333, 0x00733d, 0x90d762, 0xe2e500, 0x46d8b8, 0x4a3059, 0x12545a, 0x006536)),
+    "darkside"            : (0xbababa, 0x222324, (0x000000, 0xe8341c, 0x68c256, 0xf2d42c, 0x1c98e8, 0x8e69c9, 0x1c98e8, 0xbababa, 0x000000, 0xe05a4f, 0x77b869, 0xefd64b, 0x387cd3, 0x957bbe, 0x3d97e2, 0xbababa)),
+    "deep"                : (0xcdcdcd, 0x090909, (0x000000, 0xd70005, 0x1cd915, 0xd9bd26, 0x5665ff, 0xb052da, 0x50d2da, 0xe0e0e0, 0x535353, 0xfb0007, 0x22ff18, 0xfedc2b, 0x9fa9ff, 0xe09aff, 0x8df9ff, 0xffffff)),
+    "desert"              : (0xffffff, 0x333333, (0x4d4d4d, 0xff2b2b, 0x98fb98, 0xf0e68c, 0xcd853f, 0xffdead, 0xffa0a0, 0xf5deb3, 0x555555, 0xff5555, 0x55ff55, 0xffff55, 0x87ceff, 0xff55ff, 0xffd700, 0xffffff)),
+    "dimmedmonokai"       : (0xb9bcba, 0x1f1f1f, (0x3a3d43, 0xbe3f48, 0x879a3b, 0xc5a635, 0x4f76a1, 0x855c8d, 0x578fa4, 0xb9bcba, 0x888987, 0xfb001f, 0x0f722f, 0xc47033, 0x186de3, 0xfb0067, 0x2e706d, 0xfdffb9)),
+    "django"              : (0xf8f8f8, 0x0b2f20, (0x000000, 0xfd6209, 0x41a83e, 0xffe862, 0x245032, 0xf8f8f8, 0x9df39f, 0xffffff, 0x323232, 0xff943b, 0x73da70, 0xffff94, 0x568264, 0xffffff, 0xcfffd1, 0xffffff)),
+    "djangorebornagain"   : (0xdadedc, 0x051f14, (0x000000, 0xfd6209, 0x41a83e, 0xffe862, 0x245032, 0xf8f8f8, 0x9df39f, 0xffffff, 0x323232, 0xff943b, 0x73da70, 0xffff94, 0x568264, 0xffffff, 0xcfffd1, 0xffffff)),
+    "djangosmooth"        : (0xf8f8f8, 0x245032, (0x000000, 0xfd6209, 0x41a83e, 0xffe862, 0x989898, 0xf8f8f8, 0x9df39f, 0xe8e8e7, 0x323232, 0xff943b, 0x73da70, 0xffff94, 0xcacaca, 0xffffff, 0xcfffd1, 0xffffff)),
+    "doom.peacock"        : (0xede0ce, 0x2b2a27, (0x1c1f24, 0xcb4b16, 0x26a6a6, 0xbcd42a, 0x2a6cc6, 0xa9a1e1, 0x5699af, 0xede0ce, 0x2b2a27, 0xff5d38, 0x98be65, 0xe6f972, 0x51afef, 0xc678dd, 0x46d9ff, 0xdfdfdf)),
+    "doomone"             : (0xbbc2cf, 0x282c34, (0x000000, 0xff6c6b, 0x98be65, 0xecbe7b, 0xa9a1e1, 0xc678dd, 0x51afef, 0xbbc2cf, 0x000000, 0xff6655, 0x99bb66, 0xecbe7b, 0xa9a1e1, 0xc678dd, 0x51afef, 0xbfbfbf)),
+    "dotgov"              : (0xebebeb, 0x262c35, (0x191919, 0xbf091d, 0x3d9751, 0xf6bb34, 0x17b2e0, 0x7830b0, 0x8bd2ed, 0xffffff, 0x191919, 0xbf091d, 0x3d9751, 0xf6bb34, 0x17b2e0, 0x7830b0, 0x8bd2ed, 0xffffff)),
+    "dracula"             : (0xf8f8f2, 0x1e1f29, (0x000000, 0xff5555, 0x50fa7b, 0xf1fa8c, 0xbd93f9, 0xff79c6, 0x8be9fd, 0xbbbbbb, 0x555555, 0xff5555, 0x50fa7b, 0xf1fa8c, 0xbd93f9, 0xff79c6, 0x8be9fd, 0xffffff)),
+    "dracula+"            : (0xf8f8f2, 0x212121, (0x21222c, 0xff5555, 0x50fa7b, 0xffcb6b, 0x82aaff, 0xc792ea, 0x8be9fd, 0xf8f8f2, 0x545454, 0xff6e6e, 0x69ff94, 0xffcb6b, 0xd6acff, 0xff92df, 0xa4ffff, 0xf8f8f2)),
+    "duckbones"           : (0xebefc0, 0x0e101a, (0x0e101a, 0xe03600, 0x5dcd97, 0xe39500, 0x00a3cb, 0x795ccc, 0x00a3cb, 0xebefc0, 0x2b2f46, 0xff4821, 0x58db9e, 0xf6a100, 0x00b4e0, 0xb3a1e6, 0x00b4e0, 0xb3b692)),
+    "duotone.dark"        : (0xb7a1ff, 0x1f1d27, (0x1f1d27, 0xd9393e, 0x2dcd73, 0xd9b76e, 0xffc284, 0xde8d40, 0x2488ff, 0xb7a1ff, 0x353147, 0xd9393e, 0x2dcd73, 0xd9b76e, 0xffc284, 0xde8d40, 0x2488ff, 0xeae5ff)),
+    "earthsong"           : (0xe5c7a9, 0x292520, (0x121418, 0xc94234, 0x85c54c, 0xf5ae2e, 0x1398b9, 0xd0633d, 0x509552, 0xe5c6aa, 0x675f54, 0xff645a, 0x98e036, 0xe0d561, 0x5fdaff, 0xff9269, 0x84f088, 0xf6f7ec)),
+    "elemental"           : (0x807a74, 0x22211d, (0x3c3c30, 0x98290f, 0x479a43, 0x7f7111, 0x497f7d, 0x7f4e2f, 0x387f58, 0x807974, 0x555445, 0xe0502a, 0x61e070, 0xd69927, 0x79d9d9, 0xcd7c54, 0x59d599, 0xfff1e9)),
+    "elementary"          : (0xefefef, 0x181818, (0x242424, 0xd71c15, 0x5aa513, 0xfdb40c, 0x063b8c, 0xe40038, 0x2595e1, 0xefefef, 0x4b4b4b, 0xfc1c18, 0x6bc219, 0xfec80e, 0x0955ff, 0xfb0050, 0x3ea8fc, 0x8c00ec)),
+    "encom"               : (0x00a595, 0x000000, (0x000000, 0x9f0000, 0x008b00, 0xffd000, 0x0081ff, 0xbc00ca, 0x008b8b, 0xbbbbbb, 0x555555, 0xff0000, 0x00ee00, 0xffff00, 0x0000ff, 0xff00ff, 0x00cdcd, 0xffffff)),
+    "espresso"            : (0xffffff, 0x323232, (0x353535, 0xd25252, 0xa5c261, 0xffc66d, 0x6c99bb, 0xd197d9, 0xbed6ff, 0xeeeeec, 0x535353, 0xf00c0c, 0xc2e075, 0xe1e48b, 0x8ab7d9, 0xefb5f7, 0xdcf4ff, 0xffffff)),
+    "espresso.libre"      : (0xb8a898, 0x2a211c, (0x000000, 0xcc0000, 0x1a921c, 0xf0e53a, 0x0066ff, 0xc5656b, 0x06989a, 0xd3d7cf, 0x555753, 0xef2929, 0x9aff87, 0xfffb5c, 0x43a8ed, 0xff818a, 0x34e2e2, 0xeeeeec)),
+    "fahrenheit"          : (0xffffce, 0x000000, (0x1d1d1d, 0xcda074, 0x9e744d, 0xfecf75, 0x720102, 0x734c4d, 0x979797, 0xffffce, 0x000000, 0xfecea0, 0xcc734d, 0xfd9f4d, 0xcb4a05, 0x4e739f, 0xfed04d, 0xffffff)),
+    "fairyfloss"          : (0xf8f8f2, 0x5a5475, (0x040303, 0xf92672, 0xc2ffdf, 0xe6c000, 0xc2ffdf, 0xffb8d1, 0xc5a3ff, 0xf8f8f0, 0x6090cb, 0xff857f, 0xc2ffdf, 0xffea00, 0xc2ffdf, 0xffb8d1, 0xc5a3ff, 0xf8f8f0)),
+    "fideloper"           : (0xdbdae0, 0x292f33, (0x292f33, 0xcb1e2d, 0xedb8ac, 0xb7ab9b, 0x2e78c2, 0xc0236f, 0x309186, 0xeae3ce, 0x092028, 0xd4605a, 0xd4605a, 0xa86671, 0x7c85c4, 0x5c5db2, 0x819090, 0xfcf4df)),
+    "firefoxdev"          : (0x7c8fa4, 0x0e1011, (0x002831, 0xe63853, 0x5eb83c, 0xa57706, 0x359ddf, 0xd75cff, 0x4b73a2, 0xdcdcdc, 0x001e27, 0xe1003f, 0x1d9000, 0xcd9409, 0x006fc0, 0xa200da, 0x005794, 0xe2e2e2)),
+    "firewatch"           : (0x9ba2b2, 0x1e2027, (0x585f6d, 0xd95360, 0x5ab977, 0xdfb563, 0x4d89c4, 0xd55119, 0x44a8b6, 0xe6e5ff, 0x585f6d, 0xd95360, 0x5ab977, 0xdfb563, 0x4c89c5, 0xd55119, 0x44a8b6, 0xe6e5ff)),
+    "fishtank"            : (0xecf0fe, 0x232537, (0x03073c, 0xc6004a, 0xacf157, 0xfecd5e, 0x525fb8, 0x986f82, 0x968763, 0xecf0fc, 0x6c5b30, 0xda4b8a, 0xdbffa9, 0xfee6a9, 0xb2befa, 0xfda5cd, 0xa5bd86, 0xf6ffec)),
+    "flat"                : (0x2cc55d, 0x002240, (0x222d3f, 0xa82320, 0x32a548, 0xe58d11, 0x3167ac, 0x781aa0, 0x2c9370, 0xb0b6ba, 0x212c3c, 0xd4312e, 0x2d9440, 0xe5be0c, 0x3c7dd2, 0x8230a7, 0x35b387, 0xe7eced)),
+    "flatland"            : (0xb8dbef, 0x1d1f21, (0x1d1d19, 0xf18339, 0x9fd364, 0xf4ef6d, 0x5096be, 0x695abc, 0xd63865, 0xffffff, 0x1d1d19, 0xd22a24, 0xa7d42c, 0xff8949, 0x61b9d0, 0x695abc, 0xd63865, 0xffffff)),
+    "floraverse"          : (0xdbd1b9, 0x0e0d15, (0x08002e, 0x64002c, 0x5d731a, 0xcd751c, 0x1d6da1, 0xb7077e, 0x42a38c, 0xf3e0b8, 0x331e4d, 0xd02063, 0xb4ce59, 0xfac357, 0x40a4cf, 0xf12aae, 0x62caa8, 0xfff5db)),
+    "forestblue"          : (0xe2d8cd, 0x051519, (0x333333, 0xf8818e, 0x92d3a2, 0x1a8e63, 0x8ed0ce, 0x5e468c, 0x31658c, 0xe2d8cd, 0x3d3d3d, 0xfb3d66, 0x6bb48d, 0x30c85a, 0x39a7a2, 0x7e62b3, 0x6096bf, 0xe2d8cd)),
+    "framer"              : (0x777777, 0x111111, (0x141414, 0xff5555, 0x98ec65, 0xffcc33, 0x00aaff, 0xaa88ff, 0x88ddff, 0xcccccc, 0x414141, 0xff8888, 0xb6f292, 0xffd966, 0x33bbff, 0xcebbff, 0xbbecff, 0xffffff)),
+    "frontenddelight"     : (0xadadad, 0x1b1c1d, (0x242526, 0xf8511b, 0x565747, 0xfa771d, 0x2c70b7, 0xf02e4f, 0x3ca1a6, 0xadadad, 0x5fac6d, 0xf74319, 0x74ec4c, 0xfdc325, 0x3393ca, 0xe75e4f, 0x4fbce6, 0x8c735b)),
+    "funforrest"          : (0xdec165, 0x251200, (0x000000, 0xd6262b, 0x919c00, 0xbe8a13, 0x4699a3, 0x8d4331, 0xda8213, 0xddc265, 0x7f6a55, 0xe55a1c, 0xbfc65a, 0xffcb1b, 0x7cc9cf, 0xd26349, 0xe6a96b, 0xffeaa3)),
+    "galaxy"              : (0xffffff, 0x1d2837, (0x000000, 0xf9555f, 0x21b089, 0xfef02a, 0x589df6, 0x944d95, 0x1f9ee7, 0xbbbbbb, 0x555555, 0xfa8c8f, 0x35bb9a, 0xffff55, 0x589df6, 0xe75699, 0x3979bc, 0xffffff)),
+    "galizur"             : (0xddeeff, 0x071317, (0x223344, 0xaa1122, 0x33aa11, 0xccaa22, 0x2255cc, 0x7755aa, 0x22bbdd, 0x8899aa, 0x556677, 0xff1133, 0x33ff11, 0xffdd33, 0x3377ff, 0xaa77ff, 0x33ddff, 0xbbccdd)),
+    "github"              : (0x3e3e3e, 0xf4f4f4, (0x3e3e3e, 0x970b16, 0x07962a, 0xf8eec7, 0x003e8a, 0xe94691, 0x89d1ec, 0xffffff, 0x666666, 0xde0000, 0x87d5a2, 0xf1d007, 0x2e6cba, 0xffa29f, 0x1cfafe, 0xffffff)),
+    "github.dark"         : (0x8b949e, 0x101216, (0x000000, 0xf78166, 0x56d364, 0xe3b341, 0x6ca4f8, 0xdb61a2, 0x2b7489, 0xffffff, 0x4d4d4d, 0xf78166, 0x56d364, 0xe3b341, 0x6ca4f8, 0xdb61a2, 0x2b7489, 0xffffff)),
+    "glacier"             : (0xffffff, 0x0c1115, (0x2e343c, 0xbd0f2f, 0x35a770, 0xfb9435, 0x1f5872, 0xbd2523, 0x778397, 0xffffff, 0x404a55, 0xbd0f2f, 0x49e998, 0xfddf6e, 0x2a8bc1, 0xea4727, 0xa0b6d3, 0xffffff)),
+    "grape"               : (0x9f9fa1, 0x171423, (0x2d283f, 0xed2261, 0x1fa91b, 0x8ddc20, 0x487df4, 0x8d35c9, 0x3bdeed, 0x9e9ea0, 0x59516a, 0xf0729a, 0x53aa5e, 0xb2dc87, 0xa9bcec, 0xad81c2, 0x9de3eb, 0xa288f7)),
+    "grass"               : (0xfff0a5, 0x13773d, (0x000000, 0xbb0000, 0x00bb00, 0xe7b000, 0x0000a3, 0x950062, 0x00bbbb, 0xbbbbbb, 0x555555, 0xbb0000, 0x00bb00, 0xe7b000, 0x0000bb, 0xff55ff, 0x55ffff, 0xffffff)),
+    "grey.green"          : (0xffffff, 0x002a1a, (0x000000, 0xfe1414, 0x74ff00, 0xf1ff01, 0x00deff, 0xff00f0, 0x00ffbc, 0xffffff, 0x666666, 0xff3939, 0x00ff44, 0xffd100, 0x00afff, 0xff008a, 0x00ffd3, 0xf5ecec)),
+    "gruvbox.dark"        : (0xe6d4a3, 0x1e1e1e, (0x1e1e1e, 0xbe0f17, 0x868715, 0xcc881a, 0x377375, 0xa04b73, 0x578e57, 0x978771, 0x7f7061, 0xf73028, 0xaab01e, 0xf7b125, 0x719586, 0xc77089, 0x7db669, 0xe6d4a3)),
+    "gruvbox.light"       : (0x282828, 0xfbf1c7, (0xfbf1c7, 0x9d0006, 0x79740e, 0xb57614, 0x076678, 0x8f3f71, 0x427b58, 0x3c3836, 0x9d8374, 0xcc241d, 0x98971a, 0xd79921, 0x458588, 0xb16186, 0x689d69, 0x7c6f64)),
+    "guezwhoz"            : (0xd0d0d0, 0x1c1c1c, (0x080808, 0xff5f5f, 0x87d7af, 0xd7d787, 0x5fafd7, 0xafafff, 0x5fd7d7, 0xdadada, 0x8a8a8a, 0xd75f5f, 0xafd7af, 0xd7d7af, 0x87afd7, 0xafafd7, 0x87d7d7, 0xdadada)),
+    "hacktober"           : (0xc9c9c9, 0x141414, (0x191918, 0xb34538, 0x587744, 0xd08949, 0x206ec5, 0x864651, 0xac9166, 0xf1eee7, 0x2c2b2a, 0xb33323, 0x42824a, 0xc75a22, 0x5389c5, 0xe795a5, 0xebc587, 0xffffff)),
+    "hardcore"            : (0xa0a0a0, 0x121212, (0x1b1d1e, 0xf92672, 0xa6e22e, 0xfd971f, 0x66d9ef, 0x9e6ffe, 0x5e7175, 0xccccc6, 0x505354, 0xff669d, 0xbeed5f, 0xe6db74, 0x66d9ef, 0x9e6ffe, 0xa3babf, 0xf8f8f2)),
+    "harper"              : (0xa8a49d, 0x010101, (0x010101, 0xf8b63f, 0x7fb5e1, 0xd6da25, 0x489e48, 0xb296c6, 0xf5bfd7, 0xa8a49d, 0x726e6a, 0xf8b63f, 0x7fb5e1, 0xd6da25, 0x489e48, 0xb296c6, 0xf5bfd7, 0xfefbea)),
+    "hax0r.blue"          : (0x11b7ff, 0x010515, (0x010921, 0x10b6ff, 0x10b6ff, 0x10b6ff, 0x10b6ff, 0x10b6ff, 0x10b6ff, 0xfafafa, 0x080117, 0x00b3f7, 0x00b3f7, 0x00b3f7, 0x00b3f7, 0x00b3f7, 0x00b3f7, 0xfefefe)),
+    "hax0r.gr33n"         : (0x16b10e, 0x020f01, (0x001f0b, 0x15d00d, 0x15d00d, 0x15d00d, 0x15d00d, 0x15d00d, 0x15d00d, 0xfafafa, 0x001510, 0x19e20e, 0x19e20e, 0x19e20e, 0x19e20e, 0x19e20e, 0x19e20e, 0xfefefe)),
+    "hax0r.r3d"           : (0xb10e0e, 0x200101, (0x1f0000, 0xb00d0d, 0xb00d0d, 0xb00d0d, 0xb00d0d, 0xb00d0d, 0xb00d0d, 0xfafafa, 0x150000, 0xff1111, 0xff1010, 0xff1010, 0xff1010, 0xff1010, 0xff1010, 0xfefefe)),
+    "highway"             : (0xededed, 0x222225, (0x000000, 0xd00e18, 0x138034, 0xffcb3e, 0x006bb3, 0x6b2775, 0x384564, 0xededed, 0x5d504a, 0xf07e18, 0xb1d130, 0xfff120, 0x4fc2fd, 0xde0071, 0x5d504a, 0xffffff)),
+    "hipster.green"       : (0x84c138, 0x100b05, (0x000000, 0xb6214a, 0x00a600, 0xbfbf00, 0x246eb2, 0xb200b2, 0x00a6b2, 0xbfbfbf, 0x666666, 0xe50000, 0x86a93e, 0xe5e500, 0x0000ff, 0xe500e5, 0x00e5e5, 0xe5e5e5)),
+    "hivacruz"            : (0xede4e4, 0x132638, (0x202746, 0xc94922, 0xac9739, 0xc08b30, 0x3d8fd1, 0x6679cc, 0x22a2c9, 0x979db4, 0x6b7394, 0xc76b29, 0x73ad43, 0x5e6687, 0x898ea4, 0xdfe2f1, 0x9c637a, 0xf5f7ff)),
+    "homebrew"            : (0x00ff00, 0x000000, (0x000000, 0x990000, 0x00a600, 0x999900, 0x0000b2, 0xb200b2, 0x00a6b2, 0xbfbfbf, 0x666666, 0xe50000, 0x00d900, 0xe5e500, 0x0000ff, 0xe500e5, 0x00e5e5, 0xe5e5e5)),
+    "hopscotch"           : (0xb9b5b8, 0x322931, (0x322931, 0xdd464c, 0x8fc13e, 0xfdcc59, 0x1290bf, 0xc85e7c, 0x149b93, 0xb9b5b8, 0x797379, 0xfd8b19, 0x433b42, 0x5c545b, 0x989498, 0xd5d3d5, 0xb33508, 0xffffff)),
+    "hopscotch.256"       : (0xb9b5b8, 0x322931, (0x322931, 0xdd464c, 0x8fc13e, 0xfdcc59, 0x1290bf, 0xc85e7c, 0x149b93, 0xb9b5b8, 0x797379, 0xdd464c, 0x8fc13e, 0xfdcc59, 0x1290bf, 0xc85e7c, 0x149b93, 0xffffff)),
+    "hurtado"             : (0xdbdbdb, 0x000000, (0x575757, 0xff1b00, 0xa5e055, 0xfbe74a, 0x496487, 0xfd5ff1, 0x86e9fe, 0xcbcccb, 0x262626, 0xd51d00, 0xa5df55, 0xfbe84a, 0x89beff, 0xc001c1, 0x86eafe, 0xdbdbdb)),
+    "hybrid"              : (0xb7bcba, 0x161719, (0x2a2e33, 0xb84d51, 0xb3bf5a, 0xe4b55e, 0x6e90b0, 0xa17eac, 0x7fbfb4, 0xb5b9b6, 0x1d1f22, 0x8d2e32, 0x798431, 0xe58a50, 0x4b6b88, 0x6e5079, 0x4d7b74, 0x5a626a)),
+    "ic.green.ppl"        : (0xe0f1dc, 0x2c2c2c, (0x014401, 0xff2736, 0x41a638, 0x76a831, 0x2ec3b9, 0x50a096, 0x3ca078, 0xe6fef2, 0x035c03, 0xb4fa5c, 0xaefb86, 0xdafa87, 0x2efaeb, 0x50fafa, 0x3cfac8, 0xe0f1dc)),
+    "ic.orange.ppl"       : (0xffcb83, 0x262626, (0x000000, 0xc13900, 0xa4a900, 0xcaaf00, 0xbd6d00, 0xfc5e00, 0xf79500, 0xffc88a, 0x6a4f2a, 0xff8c68, 0xf6ff40, 0xffe36e, 0xffbe55, 0xfc874f, 0xc69752, 0xfafaff)),
+    "iceberg.dark"        : (0xc6c8d1, 0x161821, (0x1e2132, 0xe27878, 0xb4be82, 0xe2a478, 0x84a0c6, 0xa093c7, 0x89b8c2, 0xc6c8d1, 0x6b7089, 0xe98989, 0xc0ca8e, 0xe9b189, 0x91acd1, 0xada0d3, 0x95c4ce, 0xd2d4de)),
+    "iceberg.light"       : (0x33374c, 0xe8e9ec, (0xdcdfe7, 0xcc517a, 0x668e3d, 0xc57339, 0x2d539e, 0x7759b4, 0x3f83a6, 0x33374c, 0x8389a3, 0xcc3768, 0x598030, 0xb6662d, 0x22478e, 0x6845ad, 0x327698, 0x262a3f)),
+    "idea"                : (0xadadad, 0x202020, (0xadadad, 0xfc5256, 0x98b61c, 0xccb444, 0x437ee7, 0x9d74b0, 0x248887, 0x181818, 0xffffff, 0xfc7072, 0x98b61c, 0xffff0b, 0x6c9ced, 0xfc7eff, 0x248887, 0x181818)),
+    "idletoes"            : (0xffffff, 0x323232, (0x323232, 0xd25252, 0x7fe173, 0xffc66d, 0x4099ff, 0xf680ff, 0xbed6ff, 0xeeeeec, 0x535353, 0xf07070, 0x9dff91, 0xffe48b, 0x5eb7f7, 0xff9dff, 0xdcf4ff, 0xffffff)),
+    "ir.black"            : (0xf1f1f1, 0x000000, (0x4f4f4f, 0xfa6c60, 0xa8ff60, 0xfffeb7, 0x96cafe, 0xfa73fd, 0xc6c5fe, 0xefedef, 0x7b7b7b, 0xfcb6b0, 0xcfffab, 0xffffcc, 0xb5dcff, 0xfb9cfe, 0xe0e0fe, 0xffffff)),
+    "jackie.brown"        : (0xffcc2f, 0x2c1d16, (0x2c1d16, 0xef5734, 0x2baf2b, 0xbebf00, 0x246eb2, 0xd05ec1, 0x00acee, 0xbfbfbf, 0x666666, 0xe50000, 0x86a93e, 0xe5e500, 0x0000ff, 0xe500e5, 0x00e5e5, 0xe5e5e5)),
+    "japanesque"          : (0xf7f6ec, 0x1e1e1e, (0x343935, 0xcf3f61, 0x7bb75b, 0xe9b32a, 0x4c9ad4, 0xa57fc4, 0x389aad, 0xfafaf6, 0x595b59, 0xd18fa6, 0x767f2c, 0x78592f, 0x135979, 0x604291, 0x76bbca, 0xb2b5ae)),
+    "jellybeans"          : (0xdedede, 0x121212, (0x929292, 0xe27373, 0x94b979, 0xffba7b, 0x97bedc, 0xe1c0fa, 0x00988e, 0xdedede, 0xbdbdbd, 0xffa1a1, 0xbddeab, 0xffdca0, 0xb1d8f6, 0xfbdaff, 0x1ab2a8, 0xffffff)),
+    "jetbrains.darcula"   : (0xadadad, 0x202020, (0x000000, 0xfa5355, 0x126e00, 0xc2c300, 0x4581eb, 0xfa54ff, 0x33c2c1, 0xadadad, 0x555555, 0xfb7172, 0x67ff4f, 0xffff00, 0x6d9df1, 0xfb82ff, 0x60d3d1, 0xeeeeee)),
+    "jubi"                : (0xc3d3de, 0x262b33, (0x3b3750, 0xcf7b98, 0x90a94b, 0x6ebfc0, 0x576ea6, 0xbc4f68, 0x75a7d2, 0xc3d3de, 0xa874ce, 0xde90ab, 0xbcdd61, 0x87e9ea, 0x8c9fcd, 0xe16c87, 0xb7c9ef, 0xd5e5f1)),
+    "kanagawabones"       : (0xddd8bb, 0x1f1f28, (0x1f1f28, 0xe46a78, 0x98bc6d, 0xe5c283, 0x7eb3c9, 0x957fb8, 0x7eb3c9, 0xddd8bb, 0x3c3c51, 0xec818c, 0x9ec967, 0xf1c982, 0x7bc2df, 0xa98fd2, 0x7bc2df, 0xa8a48d)),
+    "kibble"              : (0xf7f7f7, 0x0e100a, (0x4d4d4d, 0xc70031, 0x29cf13, 0xd8e30e, 0x3449d1, 0x8400ff, 0x0798ab, 0xe2d1e3, 0x5a5a5a, 0xf01578, 0x6ce05c, 0xf3f79e, 0x97a4f7, 0xc495f0, 0x68f2e0, 0xffffff)),
+    "kolorit"             : (0xefecec, 0x1d1a1e, (0x1d1a1e, 0xff5b82, 0x47d7a1, 0xe8e562, 0x5db4ee, 0xda6cda, 0x57e9eb, 0xededed, 0x1d1a1e, 0xff5b82, 0x47d7a1, 0xe8e562, 0x5db4ee, 0xda6cda, 0x57e9eb, 0xededed)),
+    "konsolas"            : (0xc8c1c1, 0x060606, (0x000000, 0xaa1717, 0x18b218, 0xebae1f, 0x2323a5, 0xad1edc, 0x42b0c8, 0xc8c1c1, 0x7b716e, 0xff4141, 0x5fff5f, 0xffff55, 0x4b4bff, 0xff54ff, 0x69ffff, 0xffffff)),
+    "lab.fox"             : (0xffffff, 0x2e2e2e, (0x2e2e2e, 0xfc6d26, 0x3eb383, 0xfca121, 0xdb3b21, 0x380d75, 0x6e49cb, 0xffffff, 0x464646, 0xff6517, 0x53eaa8, 0xfca013, 0xdb501f, 0x441090, 0x7d53e7, 0xffffff)),
+    "laser"               : (0xf106e3, 0x030d18, (0x626262, 0xff8373, 0xb4fb73, 0x09b4bd, 0xfed300, 0xff90fe, 0xd1d1fe, 0xf1f1f1, 0x8f8f8f, 0xffc4be, 0xd6fcba, 0xfffed5, 0xf92883, 0xffb2fe, 0xe6e7fe, 0xffffff)),
+    "later.this.evening"  : (0x959595, 0x222222, (0x2b2b2b, 0xd45a60, 0xafba67, 0xe5d289, 0xa0bad6, 0xc092d6, 0x91bfb7, 0x3c3d3d, 0x454747, 0xd3232f, 0xaabb39, 0xe5be39, 0x6699d6, 0xab53d6, 0x5fc0ae, 0xc1c2c2)),
+    "lavandula"           : (0x736e7d, 0x050014, (0x230046, 0x7d1625, 0x337e6f, 0x7f6f49, 0x4f4a7f, 0x5a3f7f, 0x58777f, 0x736e7d, 0x372d46, 0xe05167, 0x52e0c4, 0xe0c386, 0x8e87e0, 0xa776e0, 0x9ad4e0, 0x8c91fa)),
+    "liquidcarbon"        : (0xafc2c2, 0x000000, (0x000000, 0xff3030, 0x559a70, 0xccac00, 0x0099cc, 0xcc69c8, 0x7ac4cc, 0xbccccc, 0x000000, 0xff3030, 0x559a70, 0xccac00, 0x0099cc, 0xcc69c8, 0x7ac4cc, 0xbccccc)),
+    "liquidcarbon"        : (0xafc2c2, 0x303030, (0x000000, 0xff3030, 0x559a70, 0xccac00, 0x0099cc, 0xcc69c8, 0x7ac4cc, 0xbccccc, 0x000000, 0xff3030, 0x559a70, 0xccac00, 0x0099cc, 0xcc69c8, 0x7ac4cc, 0xbccccc)),
+    "liquidcarboninverse" : (0xafc2c2, 0x000000, (0xbccccd, 0xff3030, 0x559a70, 0xccac00, 0x0099cc, 0xcc69c8, 0x7ac4cc, 0x000000, 0xffffff, 0xff3030, 0x559a70, 0xccac00, 0x0099cc, 0xcc69c8, 0x7ac4cc, 0x000000)),
+    "lovelace"            : (0xfdfdfd, 0x1d1f28, (0x282a36, 0xf37f97, 0x5adecd, 0xf2a272, 0x8897f4, 0xc574dd, 0x79e6f3, 0xfdfdfd, 0x414458, 0xff4971, 0x18e3c8, 0xff8037, 0x556fff, 0xb043d1, 0x3fdcee, 0xbebec1)),
+    "man.page"            : (0x000000, 0xfef49c, (0x000000, 0xcc0000, 0x00a600, 0x999900, 0x0000b2, 0xb200b2, 0x00a6b2, 0xcccccc, 0x666666, 0xe50000, 0x00d900, 0xe5e500, 0x0000ff, 0xe500e5, 0x00e5e5, 0xe5e5e5)),
+    "mariana"             : (0xd8dee9, 0x343d46, (0x000000, 0xec5f66, 0x99c794, 0xf9ae58, 0x6699cc, 0xc695c6, 0x5fb4b4, 0xf7f7f7, 0x333333, 0xf97b58, 0xacd1a8, 0xfac761, 0x85add6, 0xd8b6d8, 0x82c4c4, 0xffffff)),
+    "material"            : (0x232322, 0xeaeaea, (0x212121, 0xb7141f, 0x457b24, 0xf6981e, 0x134eb2, 0x560088, 0x0e717c, 0xefefef, 0x424242, 0xe83b3f, 0x7aba3a, 0xffea2e, 0x54a4f3, 0xaa4dbc, 0x26bbd1, 0xd9d9d9)),
+    "material.dark"       : (0xe5e5e5, 0x232322, (0x212121, 0xb7141f, 0x457b24, 0xf6981e, 0x134eb2, 0x560088, 0x0e717c, 0xefefef, 0x424242, 0xe83b3f, 0x7aba3a, 0xffea2e, 0x54a4f3, 0xaa4dbc, 0x26bbd1, 0xd9d9d9)),
+    "material.darker"     : (0xeeffff, 0x212121, (0x000000, 0xff5370, 0xc3e88d, 0xffcb6b, 0x82aaff, 0xc792ea, 0x89ddff, 0xffffff, 0x545454, 0xff5370, 0xc3e88d, 0xffcb6b, 0x82aaff, 0xc792ea, 0x89ddff, 0xffffff)),
+    "material.design"     : (0xe7ebed, 0x1d262a, (0x435b67, 0xfc3841, 0x5cf19e, 0xfed032, 0x37b6ff, 0xfc226e, 0x59ffd1, 0xffffff, 0xa1b0b8, 0xfc746d, 0xadf7be, 0xfee16c, 0x70cfff, 0xfc669b, 0x9affe6, 0xffffff)),
+    "material.ocean"      : (0x8f93a2, 0x0f111a, (0x546e7a, 0xff5370, 0xc3e88d, 0xffcb6b, 0x82aaff, 0xc792ea, 0x89ddff, 0xffffff, 0x546e7a, 0xff5370, 0xc3e88d, 0xffcb6b, 0x82aaff, 0xc792ea, 0x89ddff, 0xffffff)),
+    "mathias"             : (0xbbbbbb, 0x000000, (0x000000, 0xe52222, 0xa6e32d, 0xfc951e, 0xc48dff, 0xfa2573, 0x67d9f0, 0xf2f2f2, 0x555555, 0xff5555, 0x55ff55, 0xffff55, 0x5555ff, 0xff55ff, 0x55ffff, 0xffffff)),
+    "matrix"              : (0x426644, 0x0f191c, (0x0f191c, 0x23755a, 0x82d967, 0xffd700, 0x3f5242, 0x409931, 0x50b45a, 0x507350, 0x688060, 0x2fc079, 0x90d762, 0xfaff00, 0x4f7e7e, 0x11ff25, 0xc1ff8a, 0x678c61)),
+    "medallion"           : (0xcac296, 0x1d1908, (0x000000, 0xb64c00, 0x7c8b16, 0xd3bd26, 0x616bb0, 0x8c5a90, 0x916c25, 0xcac29a, 0x5e5219, 0xff9149, 0xb2ca3b, 0xffe54a, 0xacb8ff, 0xffa0ff, 0xffbc51, 0xfed698)),
+    "midnight.in.mojave"  : (0xffffff, 0x1e1e1e, (0x1e1e1e, 0xff453a, 0x32d74b, 0xffd60a, 0x0a84ff, 0xbf5af2, 0x5ac8fa, 0xffffff, 0x1e1e1e, 0xff453a, 0x32d74b, 0xffd60a, 0x0a84ff, 0xbf5af2, 0x5ac8fa, 0xffffff)),
+    "mirage"              : (0xa6b2c0, 0x1b2738, (0x011627, 0xff9999, 0x85cc95, 0xffd700, 0x7fb5ff, 0xddb3ff, 0x21c7a8, 0xffffff, 0x575656, 0xff9999, 0x85cc95, 0xffd700, 0x7fb5ff, 0xddb3ff, 0x85cc95, 0xffffff)),
+    "misterioso"          : (0xe1e1e0, 0x2d3743, (0x000000, 0xff4242, 0x74af68, 0xffad29, 0x338f86, 0x9414e6, 0x23d7d7, 0xe1e1e0, 0x555555, 0xff3242, 0x74cd68, 0xffb929, 0x23d7d7, 0xff37ff, 0x00ede1, 0xffffff)),
+    "molokai"             : (0xbbbbbb, 0x121212, (0x121212, 0xfa2573, 0x98e123, 0xdfd460, 0x1080d0, 0x8700ff, 0x43a8d0, 0xbbbbbb, 0x555555, 0xf6669d, 0xb1e05f, 0xfff26d, 0x00afff, 0xaf87ff, 0x51ceff, 0xffffff)),
+    "monalisa"            : (0xf7d66a, 0x120b0d, (0x351b0e, 0x9b291c, 0x636232, 0xc36e28, 0x515c5d, 0x9b1d29, 0x588056, 0xf7d75c, 0x874228, 0xff4331, 0xb4b264, 0xff9566, 0x9eb2b4, 0xff5b6a, 0x8acd8f, 0xffe598)),
+    "monokai.remastered"  : (0xd9d9d9, 0x0c0c0c, (0x1a1a1a, 0xf4005f, 0x98e024, 0xfd971f, 0x9d65ff, 0xf4005f, 0x58d1eb, 0xc4c5b5, 0x625e4c, 0xf4005f, 0x98e024, 0xe0d561, 0x9d65ff, 0xf4005f, 0x58d1eb, 0xf6f6ef)),
+    "monokai.soda"        : (0xc4c5b5, 0x1a1a1a, (0x1a1a1a, 0xf4005f, 0x98e024, 0xfa8419, 0x9d65ff, 0xf4005f, 0x58d1eb, 0xc4c5b5, 0x625e4c, 0xf4005f, 0x98e024, 0xe0d561, 0x9d65ff, 0xf4005f, 0x58d1eb, 0xf6f6ef)),
+    "monokai.vivid"       : (0xf9f9f9, 0x121212, (0x121212, 0xfa2934, 0x98e123, 0xfff30a, 0x0443ff, 0xf800f8, 0x01b6ed, 0xffffff, 0x838383, 0xf6669d, 0xb1e05f, 0xfff26d, 0x0443ff, 0xf200f6, 0x51ceff, 0xffffff)),
+    "n0tch2k"             : (0xa0a0a0, 0x222222, (0x383838, 0xa95551, 0x666666, 0xa98051, 0x657d3e, 0x767676, 0xc9c9c9, 0xd0b8a3, 0x474747, 0xa97775, 0x8c8c8c, 0xa99175, 0x98bd5e, 0xa3a3a3, 0xdcdcdc, 0xd8c8bb)),
+    "neobones.dark"       : (0xc6d5cf, 0x0f191f, (0x0f191f, 0xde6e7c, 0x90ff6b, 0xb77e64, 0x8190d4, 0xb279a7, 0x66a5ad, 0xc6d5cf, 0x263945, 0xe8838f, 0xa0ff85, 0xd68c67, 0x92a0e2, 0xcf86c1, 0x65b8c1, 0x98a39e)),
+    "neobones.light"      : (0x202e18, 0xe5ede6, (0xe5ede6, 0xa8334c, 0x567a30, 0x944927, 0x286486, 0x88507d, 0x3b8992, 0x202e18, 0xb3c6b6, 0x94253e, 0x3f5a22, 0x803d1c, 0x1d5573, 0x7b3b70, 0x2b747c, 0x415934)),
+    "neon"                : (0x00fffc, 0x14161a, (0x000000, 0xff3045, 0x5ffa74, 0xfffc7e, 0x0208cb, 0xf924e7, 0x00fffc, 0xc7c7c7, 0x686868, 0xff5a5a, 0x75ff88, 0xfffd96, 0x3c40cb, 0xf15be5, 0x88fffe, 0xffffff)),
+    "neopolitan"          : (0xffffff, 0x271f19, (0x000000, 0x800000, 0x61ce3c, 0xfbde2d, 0x253b76, 0xff0080, 0x8da6ce, 0xf8f8f8, 0x000000, 0x800000, 0x61ce3c, 0xfbde2d, 0x253b76, 0xff0080, 0x8da6ce, 0xf8f8f8)),
+    "neutron"             : (0xe6e8ef, 0x1c1e22, (0x23252b, 0xb54036, 0x5ab977, 0xdeb566, 0x6a7c93, 0xa4799d, 0x3f94a8, 0xe6e8ef, 0x23252b, 0xb54036, 0x5ab977, 0xdeb566, 0x6a7c93, 0xa4799d, 0x3f94a8, 0xebedf2)),
+    "night.owlish.light"  : (0x403f53, 0xffffff, (0x011627, 0xd3423e, 0x2aa298, 0xdaaa01, 0x4876d6, 0x403f53, 0x08916a, 0x7a8181, 0x7a8181, 0xf76e6e, 0x49d0c5, 0xdac26b, 0x5ca7e4, 0x697098, 0x00c990, 0x989fb1)),
+    "nightlion.v1"        : (0xbbbbbb, 0x000000, (0x4c4c4c, 0xbb0000, 0x5fde8f, 0xf3f167, 0x276bd8, 0xbb00bb, 0x00dadf, 0xbbbbbb, 0x555555, 0xff5555, 0x55ff55, 0xffff55, 0x5555ff, 0xff55ff, 0x55ffff, 0xffffff)),
+    "nightlion.v2"        : (0xbbbbbb, 0x171717, (0x4c4c4c, 0xbb0000, 0x04f623, 0xf3f167, 0x64d0f0, 0xce6fdb, 0x00dadf, 0xbbbbbb, 0x555555, 0xff5555, 0x7df71d, 0xffff55, 0x62cbe8, 0xff9bf5, 0x00ccd8, 0xffffff)),
+    "nocturnal.winter"    : (0xe6e5e5, 0x0d0d17, (0x4d4d4d, 0xf12d52, 0x09cd7e, 0xf5f17a, 0x3182e0, 0xff2b6d, 0x09c87a, 0xfcfcfc, 0x808080, 0xf16d86, 0x0ae78d, 0xfffc67, 0x6096ff, 0xff78a2, 0x0ae78d, 0xffffff)),
+    "nord"                : (0xd8dee9, 0x2e3440, (0x3b4252, 0xbf616a, 0xa3be8c, 0xebcb8b, 0x81a1c1, 0xb48ead, 0x88c0d0, 0xe5e9f0, 0x4c566a, 0xbf616a, 0xa3be8c, 0xebcb8b, 0x81a1c1, 0xb48ead, 0x8fbcbb, 0xeceff4)),
+    "nord.light"          : (0x414858, 0xe5e9f0, (0x3b4252, 0xbf616a, 0xa3be8c, 0xebcb8b, 0x81a1c1, 0xb48ead, 0x88c0d0, 0xd8dee9, 0x4c566a, 0xbf616a, 0xa3be8c, 0xebcb8b, 0x81a1c1, 0xb48ead, 0x8fbcbb, 0xeceff4)),
+    "novel"               : (0x3b2322, 0xdfdbc3, (0x000000, 0xcc0000, 0x009600, 0xd06b00, 0x0000cc, 0xcc00cc, 0x0087cc, 0xcccccc, 0x808080, 0xcc0000, 0x009600, 0xd06b00, 0x0000cc, 0xcc00cc, 0x0087cc, 0xffffff)),
+    "obsidian"            : (0xcdcdcd, 0x283033, (0x000000, 0xa60001, 0x00bb00, 0xfecd22, 0x3a9bdb, 0xbb00bb, 0x00bbbb, 0xbbbbbb, 0x555555, 0xff0003, 0x93c863, 0xfef874, 0xa1d7ff, 0xff55ff, 0x55ffff, 0xffffff)),
+    "ocean"               : (0xffffff, 0x224fbc, (0x000000, 0x990000, 0x00a600, 0x999900, 0x0000b2, 0xb200b2, 0x00a6b2, 0xbfbfbf, 0x666666, 0xe50000, 0x00d900, 0xe5e500, 0x0000ff, 0xe500e5, 0x00e5e5, 0xe5e5e5)),
+    "oceanic.next"        : (0xb3b8c3, 0x121b21, (0x121c21, 0xe44754, 0x89bd82, 0xf7bd51, 0x5486c0, 0xb77eb8, 0x50a5a4, 0xffffff, 0x52606b, 0xe44754, 0x89bd82, 0xf7bd51, 0x5486c0, 0xb77eb8, 0x50a5a4, 0xffffff)),
+    "oceanicmaterial"     : (0xc2c8d7, 0x1c262b, (0x000000, 0xee2b2a, 0x40a33f, 0xffea2e, 0x1e80f0, 0x8800a0, 0x16afca, 0xa4a4a4, 0x777777, 0xdc5c60, 0x70be71, 0xfff163, 0x54a4f3, 0xaa4dbc, 0x42c7da, 0xffffff)),
+    "ollie"               : (0x8a8dae, 0x222125, (0x000000, 0xac2e31, 0x31ac61, 0xac4300, 0x2d57ac, 0xb08528, 0x1fa6ac, 0x8a8eac, 0x5b3725, 0xff3d48, 0x3bff99, 0xff5e1e, 0x4488ff, 0xffc21d, 0x1ffaff, 0x5b6ea7)),
+    "onehalf.dark"        : (0xdcdfe4, 0x282c34, (0x282c34, 0xe06c75, 0x98c379, 0xe5c07b, 0x61afef, 0xc678dd, 0x56b6c2, 0xdcdfe4, 0x282c34, 0xe06c75, 0x98c379, 0xe5c07b, 0x61afef, 0xc678dd, 0x56b6c2, 0xdcdfe4)),
+    "onehalf.light"       : (0x383a42, 0xfafafa, (0x383a42, 0xe45649, 0x50a14f, 0xc18401, 0x0184bc, 0xa626a4, 0x0997b3, 0xfafafa, 0x4f525e, 0xe06c75, 0x98c379, 0xe5c07b, 0x61afef, 0xc678dd, 0x56b6c2, 0xffffff)),
+    "operator.mono.dark"  : (0xc3cac2, 0x191919, (0x5a5a5a, 0xca372d, 0x4d7b3a, 0xd4d697, 0x4387cf, 0xb86cb4, 0x72d5c6, 0xced4cd, 0x9a9b99, 0xc37d62, 0x83d0a2, 0xfdfdc5, 0x89d3f6, 0xff2c7a, 0x82eada, 0xfdfdf6)),
+    "overnight.slumber"   : (0xced2d6, 0x0e1729, (0x0a1222, 0xffa7c4, 0x85cc95, 0xffcb8b, 0x8dabe1, 0xc792eb, 0x78ccf0, 0xffffff, 0x575656, 0xffa7c4, 0x85cc95, 0xffcb8b, 0x8dabe1, 0xc792eb, 0xffa7c4, 0xffffff)),
+    "palenighthc"         : (0xcccccc, 0x3e4251, (0x000000, 0xf07178, 0xc3e88d, 0xffcb6b, 0x82aaff, 0xc792ea, 0x89ddff, 0xffffff, 0x666666, 0xf6a9ae, 0xdbf1ba, 0xffdfa6, 0xb4ccff, 0xddbdf2, 0xb8eaff, 0x999999)),
+    "pandora"             : (0xe1e1e1, 0x141e43, (0x000000, 0xff4242, 0x74af68, 0xffad29, 0x338f86, 0x9414e6, 0x23d7d7, 0xe2e2e2, 0x3f5648, 0xff3242, 0x74cd68, 0xffb929, 0x23d7d7, 0xff37ff, 0x00ede1, 0xffffff)),
+    "paraiso.dark"        : (0xa39e9b, 0x2f1e2e, (0x2f1e2e, 0xef6155, 0x48b685, 0xfec418, 0x06b6ef, 0x815ba4, 0x5bc4bf, 0xa39e9b, 0x776e71, 0xef6155, 0x48b685, 0xfec418, 0x06b6ef, 0x815ba4, 0x5bc4bf, 0xe7e9db)),
+    "paulmillr"           : (0xf2f2f2, 0x000000, (0x2a2a2a, 0xff0000, 0x79ff0f, 0xe7bf00, 0x396bd7, 0xb449be, 0x66ccff, 0xbbbbbb, 0x666666, 0xff0080, 0x66ff66, 0xf3d64e, 0x709aed, 0xdb67e6, 0x7adff2, 0xffffff)),
+    "pencil.dark"         : (0xf1f1f1, 0x212121, (0x212121, 0xc30771, 0x10a778, 0xa89c14, 0x008ec4, 0x523c79, 0x20a5ba, 0xd9d9d9, 0x424242, 0xfb007a, 0x5fd7af, 0xf3e430, 0x20bbfc, 0x6855de, 0x4fb8cc, 0xf1f1f1)),
+    "pencil.light"        : (0x424242, 0xf1f1f1, (0x212121, 0xc30771, 0x10a778, 0xa89c14, 0x008ec4, 0x523c79, 0x20a5ba, 0xd9d9d9, 0x424242, 0xfb007a, 0x5fd7af, 0xf3e430, 0x20bbfc, 0x6855de, 0x4fb8cc, 0xf1f1f1)),
+    "peppermint"          : (0xc8c8c8, 0x000000, (0x353535, 0xe74669, 0x89d287, 0xdab853, 0x449fd0, 0xda62dc, 0x65aaaf, 0xb4b4b4, 0x535353, 0xe4859b, 0xa3cca2, 0xe1e487, 0x6fbce2, 0xe586e7, 0x96dcdb, 0xdfdfdf)),
+    "piatto.light"        : (0x414141, 0xffffff, (0x414141, 0xb23771, 0x66781e, 0xcd6f34, 0x3c5ea8, 0xa454b2, 0x66781e, 0xffffff, 0x3f3f3f, 0xdb3365, 0x829429, 0xcd6f34, 0x3c5ea8, 0xa454b2, 0x829429, 0xf2f2f2)),
+    "pnevma"              : (0xd0d0d0, 0x1c1c1c, (0x2f2e2d, 0xa36666, 0x90a57d, 0xd7af87, 0x7fa5bd, 0xc79ec4, 0x8adbb4, 0xd0d0d0, 0x4a4845, 0xd78787, 0xafbea2, 0xe4c9af, 0xa1bdce, 0xd7beda, 0xb1e7dd, 0xefefef)),
+    "pop.and.lock"        : (0xebdbb2, 0x181921, (0x1d2021, 0xcc241d, 0x98971a, 0xd79921, 0x458588, 0xb16286, 0x689d6a, 0xa89984, 0x928374, 0xf42c3e, 0xb8bb26, 0xfabd2f, 0x99c6ca, 0xd3869b, 0x7ec16e, 0xebdbb2)),
+    "primary"             : (0x000000, 0xffffff, (0x000000, 0xdb4437, 0x0f9d58, 0xf4b400, 0x4285f4, 0xdb4437, 0x4285f4, 0xffffff, 0x000000, 0xdb4437, 0x0f9d58, 0xf4b400, 0x4285f4, 0x4285f4, 0x0f9d58, 0xffffff)),
+    "pro"                 : (0xf2f2f2, 0x000000, (0x000000, 0x990000, 0x00a600, 0x999900, 0x2009db, 0xb200b2, 0x00a6b2, 0xbfbfbf, 0x666666, 0xe50000, 0x00d900, 0xe5e500, 0x0000ff, 0xe500e5, 0x00e5e5, 0xe5e5e5)),
+    "pro.light"           : (0x191919, 0xffffff, (0x000000, 0xe5492b, 0x50d148, 0xc6c440, 0x3b75ff, 0xed66e8, 0x4ed2de, 0xdcdcdc, 0x9f9f9f, 0xff6640, 0x61ef57, 0xf2f156, 0x0082ff, 0xff7eff, 0x61f7f8, 0xf2f2f2)),
+    "purple.rain"         : (0xfffbf6, 0x21084a, (0x000000, 0xff260e, 0x9be205, 0xffc400, 0x00a2fa, 0x815bb5, 0x00deef, 0xffffff, 0x565656, 0xff4250, 0xb8e36e, 0xffd852, 0x00a6ff, 0xac7bf0, 0x74fdf3, 0xffffff)),
+    "purplepeter"         : (0xece7fa, 0x2a1a4a, (0x0a0520, 0xff796d, 0x99b481, 0xefdfac, 0x66d9ef, 0xe78fcd, 0xba8cff, 0xffba81, 0x100b23, 0xf99f92, 0xb4be8f, 0xf2e9bf, 0x79daed, 0xba91d4, 0xa0a0d6, 0xb9aed3)),
+    "rapture"             : (0xc0c9e5, 0x111e2a, (0x000000, 0xfc644d, 0x7afde1, 0xfff09b, 0x6c9bf5, 0xff4fa1, 0x64e0ff, 0xc0c9e5, 0x304b66, 0xfc644d, 0x7afde1, 0xfff09b, 0x6c9bf5, 0xff4fa1, 0x64e0ff, 0xffffff)),
+    "raycast.dark"        : (0xffffff, 0x1a1a1a, (0x000000, 0xff5360, 0x59d499, 0xffc531, 0x56c2ff, 0xcf2f98, 0x52eee5, 0xffffff, 0x000000, 0xff6363, 0x59d499, 0xffc531, 0x56c2ff, 0xcf2f98, 0x52eee5, 0xffffff)),
+    "raycast.light"       : (0x000000, 0xffffff, (0x000000, 0xb12424, 0x006b4f, 0xf8a300, 0x138af2, 0x9a1b6e, 0x3eb8bf, 0xffffff, 0x000000, 0xb12424, 0x006b4f, 0xf8a300, 0x138af2, 0x9a1b6e, 0x3eb8bf, 0xffffff)),
+    "rebecca"             : (0xe8e6ed, 0x292a44, (0x12131e, 0xdd7755, 0x04dbb5, 0xf2e7b7, 0x7aa5ff, 0xbf9cf9, 0x56d3c2, 0xe4e3e9, 0x666699, 0xff92cd, 0x01eac0, 0xfffca8, 0x69c0fa, 0xc17ff8, 0x8bfde1, 0xf4f2f9)),
+    "red.alert"           : (0xffffff, 0x762423, (0x000000, 0xd62e4e, 0x71be6b, 0xbeb86b, 0x489bee, 0xe979d7, 0x6bbeb8, 0xd6d6d6, 0x262626, 0xe02553, 0xaff08c, 0xdfddb7, 0x65aaf1, 0xddb7df, 0xb7dfdd, 0xffffff)),
+    "red.planet"          : (0xc2b790, 0x222222, (0x202020, 0x8c3432, 0x728271, 0xe8bf6a, 0x69819e, 0x896492, 0x5b8390, 0xb9aa99, 0x676767, 0xb55242, 0x869985, 0xebeb91, 0x60827e, 0xde4974, 0x38add8, 0xd6bfb8)),
+    "red.sands"           : (0xd7c9a7, 0x7a251e, (0x000000, 0xff3f00, 0x00bb00, 0xe7b000, 0x0072ff, 0xbb00bb, 0x00bbbb, 0xbbbbbb, 0x555555, 0xbb0000, 0x00bb00, 0xe7b000, 0x0072ae, 0xff55ff, 0x55ffff, 0xffffff)),
+    "relaxed"             : (0xd9d9d9, 0x353a44, (0x151515, 0xbc5653, 0x909d63, 0xebc17a, 0x6a8799, 0xb06698, 0xc9dfff, 0xd9d9d9, 0x636363, 0xbc5653, 0xa0ac77, 0xebc17a, 0x7eaac7, 0xb06698, 0xacbbd0, 0xf7f7f7)),
+    "retro"               : (0x13a10e, 0x000000, (0x13a10e, 0x13a10e, 0x13a10e, 0x13a10e, 0x13a10e, 0x13a10e, 0x13a10e, 0x13a10e, 0x16ba10, 0x16ba10, 0x16ba10, 0x16ba10, 0x16ba10, 0x16ba10, 0x16ba10, 0x16ba10)),
+    "rippedcasts"         : (0xffffff, 0x2b2b2b, (0x000000, 0xcdaf95, 0xa8ff60, 0xbfbb1f, 0x75a5b0, 0xff73fd, 0x5a647e, 0xbfbfbf, 0x666666, 0xeecbad, 0xbcee68, 0xe5e500, 0x86bdc9, 0xe500e5, 0x8c9bc4, 0xe5e5e5)),
+    "rouge.2"             : (0xa2a3aa, 0x17182b, (0x5d5d6b, 0xc6797e, 0x969e92, 0xdbcdab, 0x6e94b9, 0x4c4e78, 0x8ab6c1, 0xe8e8ea, 0x616274, 0xc6797e, 0xe6dcc4, 0xe6dcc4, 0x98b3cd, 0x8283a1, 0xabcbd3, 0xe8e8ea)),
+    "royal"               : (0x514968, 0x100815, (0x241f2b, 0x91284c, 0x23801c, 0xb49d27, 0x6580b0, 0x674d96, 0x8aaabe, 0x524966, 0x312d3d, 0xd5356c, 0x2cd946, 0xfde83b, 0x90baf9, 0xa479e3, 0xacd4eb, 0x9e8cbd)),
+    "ryuuko"              : (0xececec, 0x2c3941, (0x2c3941, 0x865f5b, 0x66907d, 0xb1a990, 0x6a8e95, 0xb18a73, 0x88b2ac, 0xececec, 0x5d7079, 0x865f5b, 0x66907d, 0xb1a990, 0x6a8e95, 0xb18a73, 0x88b2ac, 0xececec)),
+    "sakura"              : (0xdd7bdc, 0x18131e, (0x000000, 0xd52370, 0x41af1a, 0xbc7053, 0x6964ab, 0xc71fbf, 0x939393, 0x998eac, 0x786d69, 0xf41d99, 0x22e529, 0xf59574, 0x9892f1, 0xe90cdd, 0xeeeeee, 0xcbb6ff)),
+    "scarlet.protocol"    : (0xe41951, 0x1c153d, (0x101116, 0xff0051, 0x00dc84, 0xfaf945, 0x0271b6, 0xca30c7, 0x00c5c7, 0xc7c7c7, 0x686868, 0xff6e67, 0x5ffa68, 0xfffc67, 0x6871ff, 0xbd35ec, 0x60fdff, 0xffffff)),
+    "seafoam.pastel"      : (0xd4e7d4, 0x243435, (0x757575, 0x825d4d, 0x728c62, 0xada16d, 0x4d7b82, 0x8a7267, 0x729494, 0xe0e0e0, 0x8a8a8a, 0xcf937a, 0x98d9aa, 0xfae79d, 0x7ac3cf, 0xd6b2a1, 0xade0e0, 0xe0e0e0)),
+    "seashells"           : (0xdeb88d, 0x09141b, (0x17384c, 0xd15123, 0x027c9b, 0xfca02f, 0x1e4950, 0x68d4f1, 0x50a3b5, 0xdeb88d, 0x434b53, 0xd48678, 0x628d98, 0xfdd39f, 0x1bbcdd, 0xbbe3ee, 0x87acb4, 0xfee4ce)),
+    "seoulbones.dark"     : (0xdddddd, 0x4b4b4b, (0x4b4b4b, 0xe388a3, 0x98bd99, 0xffdf9b, 0x97bdde, 0xa5a6c5, 0x6fbdbe, 0xdddddd, 0x6c6465, 0xeb99b1, 0x8fcd92, 0xffe5b3, 0xa2c8e9, 0xb2b3da, 0x6bcacb, 0xa8a8a8)),
+    "seoulbones.light"    : (0x555555, 0xe2e2e2, (0xe2e2e2, 0xdc5284, 0x628562, 0xc48562, 0x0084a3, 0x896788, 0x008586, 0x555555, 0xbfbabb, 0xbe3c6d, 0x487249, 0xa76b48, 0x006f89, 0x7f4c7e, 0x006f70, 0x777777)),
+    "seti"                : (0xcacecd, 0x111213, (0x323232, 0xc22832, 0x8ec43d, 0xe0c64f, 0x43a5d5, 0x8b57b5, 0x8ec43d, 0xeeeeee, 0x323232, 0xc22832, 0x8ec43d, 0xe0c64f, 0x43a5d5, 0x8b57b5, 0x8ec43d, 0xffffff)),
+    "shades.of.purple"    : (0xffffff, 0x1e1d40, (0x000000, 0xd90429, 0x3ad900, 0xffe700, 0x6943ff, 0xff2c70, 0x00c5c7, 0xc7c7c7, 0x686868, 0xf92a1c, 0x43d426, 0xf1d000, 0x6871ff, 0xff77ff, 0x79e8fb, 0xffffff)),
+    "shaman"              : (0x405555, 0x001015, (0x012026, 0xb2302d, 0x00a941, 0x5e8baa, 0x449a86, 0x00599d, 0x5d7e19, 0x405555, 0x384451, 0xff4242, 0x2aea5e, 0x8ed4fd, 0x61d5ba, 0x1298ff, 0x98d028, 0x58fbd6)),
+    "slate"               : (0x35b1d2, 0x222222, (0x222222, 0xe2a8bf, 0x81d778, 0xc4c9c0, 0x264b49, 0xa481d3, 0x15ab9c, 0x02c5e0, 0xffffff, 0xffcdd9, 0xbeffa8, 0xd0ccca, 0x7ab0d2, 0xc5a7d9, 0x8cdfe0, 0xe0e0e0)),
+    "sleepyhollow"        : (0xaf9a91, 0x121214, (0x572100, 0xba3934, 0x91773f, 0xb55600, 0x5f63b4, 0xa17c7b, 0x8faea9, 0xaf9a91, 0x4e4b61, 0xd9443f, 0xd6b04e, 0xf66813, 0x8086ef, 0xe2c2bb, 0xa4dce7, 0xd2c7a9)),
+    "smyck"               : (0xf7f7f7, 0x1b1b1b, (0x000000, 0xb84131, 0x7da900, 0xc4a500, 0x62a3c4, 0xba8acc, 0x207383, 0xa1a1a1, 0x7a7a7a, 0xd6837c, 0xc4f137, 0xfee14d, 0x8dcff0, 0xf79aff, 0x6ad9cf, 0xf7f7f7)),
+    "snazzy"              : (0xebece6, 0x1e1f29, (0x000000, 0xfc4346, 0x50fb7c, 0xf0fb8c, 0x49baff, 0xfc4cb4, 0x8be9fe, 0xededec, 0x555555, 0xfc4346, 0x50fb7c, 0xf0fb8c, 0x49baff, 0xfc4cb4, 0x8be9fe, 0xededec)),
+    "softserver"          : (0x99a3a2, 0x242626, (0x000000, 0xa2686a, 0x9aa56a, 0xa3906a, 0x6b8fa3, 0x6a71a3, 0x6ba58f, 0x99a3a2, 0x666c6c, 0xdd5c60, 0xbfdf55, 0xdeb360, 0x62b1df, 0x606edf, 0x64e39c, 0xd2e0de)),
+    "solarized.darcula"   : (0xd2d8d9, 0x3d3f41, (0x25292a, 0xf24840, 0x629655, 0xb68800, 0x2075c7, 0x797fd4, 0x15968d, 0xd2d8d9, 0x25292a, 0xf24840, 0x629655, 0xb68800, 0x2075c7, 0x797fd4, 0x15968d, 0xd2d8d9)),
+    "solarized.dark.c"    : (0x9cc2c3, 0x001e27, (0x002831, 0xd11c24, 0x6cbe6c, 0xa57706, 0x2176c7, 0xc61c6f, 0x259286, 0xeae3cb, 0x006488, 0xf5163b, 0x51ef84, 0xb27e28, 0x178ec8, 0xe24d8e, 0x00b39e, 0xfcf4dc)),
+    "solarized.dark.p"    : (0x708284, 0x001e27, (0x002831, 0xd11c24, 0x738a05, 0xa57706, 0x2176c7, 0xc61c6f, 0x259286, 0xeae3cb, 0x475b62, 0xbd3613, 0x475b62, 0x536870, 0x708284, 0x5956ba, 0x819090, 0xfcf4dc)),
+    "spacedust"           : (0xecf0c1, 0x0a1e24, (0x6e5346, 0xe35b00, 0x5cab96, 0xe3cd7b, 0x0f548b, 0xe35b00, 0x06afc7, 0xf0f1ce, 0x684c31, 0xff8a3a, 0xaecab8, 0xffc878, 0x67a0ce, 0xff8a3a, 0x83a7b4, 0xfefff1)),
+    "spacegray"           : (0xb3b8c3, 0x20242d, (0x000000, 0xb04b57, 0x87b379, 0xe5c179, 0x7d8fa4, 0xa47996, 0x85a7a5, 0xb3b8c3, 0x000000, 0xb04b57, 0x87b379, 0xe5c179, 0x7d8fa4, 0xa47996, 0x85a7a5, 0xffffff)),
+    "spacegray.1980"      : (0xbdbaae, 0x222222, (0x15171c, 0xec5f67, 0x81a764, 0xfec254, 0x5486c0, 0xbf83c1, 0x57c2c1, 0xefece7, 0x555555, 0xff6973, 0x93d493, 0xffd256, 0x4d84d1, 0xff55ff, 0x83e9e4, 0xffffff)),
+    "spacegray.1980.dull" : (0xc9c6bc, 0x222222, (0x15171c, 0xb24a56, 0x92b477, 0xc6735a, 0x7c8fa5, 0xa5789e, 0x80cdcb, 0xb3b8c3, 0x555555, 0xec5f67, 0x89e986, 0xfec254, 0x5486c0, 0xbf83c1, 0x58c2c1, 0xffffff)),
+    "spiderman"           : (0xe3e3e3, 0x1b1d1e, (0x1b1d1e, 0xe60813, 0xe22928, 0xe24756, 0x2c3fff, 0x2435db, 0x3256ff, 0xfffef6, 0x505354, 0xff0325, 0xff3338, 0xfe3a35, 0x1d50ff, 0x747cff, 0x6184ff, 0xfffff9)),
+    "spring"              : (0x4d4d4c, 0xffffff, (0x000000, 0xff4d83, 0x1f8c3b, 0x1fc95b, 0x1dd3ee, 0x8959a8, 0x3e999f, 0xffffff, 0x000000, 0xff0021, 0x1fc231, 0xd5b807, 0x15a9fd, 0x8959a8, 0x3e999f, 0xffffff)),
+    "square"              : (0xacacab, 0x1a1a1a, (0x050505, 0xe9897c, 0xb6377d, 0xecebbe, 0xa9cdeb, 0x75507b, 0xc9caec, 0xf2f2f2, 0x141414, 0xf99286, 0xc3f786, 0xfcfbcc, 0xb6defb, 0xad7fa8, 0xd7d9fc, 0xe2e2e2)),
+    "sublette"            : (0xccced0, 0x202535, (0x253045, 0xee5577, 0x55ee77, 0xffdd88, 0x5588ff, 0xff77cc, 0x44eeee, 0xf5f5da, 0x405570, 0xee6655, 0x99ee77, 0xffff77, 0x77bbff, 0xaa88ff, 0x55ffbb, 0xffffee)),
+    "subliminal"          : (0xd4d4d4, 0x282c35, (0x7f7f7f, 0xe15a60, 0xa9cfa4, 0xffe2a9, 0x6699cc, 0xf1a5ab, 0x5fb3b3, 0xd4d4d4, 0x7f7f7f, 0xe15a60, 0xa9cfa4, 0xffe2a9, 0x6699cc, 0xf1a5ab, 0x5fb3b3, 0xd4d4d4)),
+    "sundried"            : (0xc9c9c9, 0x1a1818, (0x302b2a, 0xa7463d, 0x587744, 0x9d602a, 0x485b98, 0x864651, 0x9c814f, 0xc9c9c9, 0x4d4e48, 0xaa000c, 0x128c21, 0xfc6a21, 0x7999f7, 0xfd8aa1, 0xfad484, 0xffffff)),
+    "symfonic"            : (0xffffff, 0x000000, (0x000000, 0xdc322f, 0x56db3a, 0xff8400, 0x0084d4, 0xb729d9, 0xccccff, 0xffffff, 0x1b1d21, 0xdc322f, 0x56db3a, 0xff8400, 0x0084d4, 0xb729d9, 0xccccff, 0xffffff)),
+    "synthwave"           : (0xdad9c7, 0x000000, (0x000000, 0xf6188f, 0x1ebb2b, 0xfdf834, 0x2186ec, 0xf85a21, 0x12c3e2, 0xffffff, 0x000000, 0xf841a0, 0x25c141, 0xfdf454, 0x2f9ded, 0xf97137, 0x19cde6, 0xffffff)),
+    "synthwave.all"       : (0xf0eff1, 0x2a2139, (0xfefefe, 0xf97e72, 0x72f1b8, 0xfede5d, 0x6d77b3, 0xc792ea, 0xf772e0, 0xfefefe, 0xfefefe, 0xf88414, 0x72f1b8, 0xfff951, 0x36f9f6, 0xe1acff, 0xf92aad, 0xfefefe)),
+    "tango.adapted"       : (0x000000, 0xffffff, (0x000000, 0xff0000, 0x59d600, 0xf0cb00, 0x00a2ff, 0xc17ecc, 0x00d0d6, 0xe6ebe1, 0x8f928b, 0xff0013, 0x93ff00, 0xfff121, 0x88c9ff, 0xe9a7e1, 0x00feff, 0xf6f6f4)),
+    "tango.half.adapted"  : (0x000000, 0xffffff, (0x000000, 0xff0000, 0x4cc300, 0xe2c000, 0x008ef6, 0xa96cb3, 0x00bdc3, 0xe0e5db, 0x797d76, 0xff0013, 0x8af600, 0xffec00, 0x76bfff, 0xd898d1, 0x00f6fa, 0xf4f4f2)),
+    "teerb"               : (0xd0d0d0, 0x262626, (0x1c1c1c, 0xd68686, 0xaed686, 0xd7af87, 0x86aed6, 0xd6aed6, 0x8adbb4, 0xd0d0d0, 0x1c1c1c, 0xd68686, 0xaed686, 0xe4c9af, 0x86aed6, 0xd6aed6, 0xb1e7dd, 0xefefef)),
+    "terminal.basic"      : (0x000000, 0xffffff, (0x000000, 0x990000, 0x00a600, 0x999900, 0x0000b2, 0xb200b2, 0x00a6b2, 0xbfbfbf, 0x666666, 0xe50000, 0x00d900, 0xe5e500, 0x0000ff, 0xe500e5, 0x00e5e5, 0xe5e5e5)),
+    "thayer.bright"       : (0xf8f8f8, 0x1b1d1e, (0x1b1d1e, 0xf92672, 0x4df840, 0xf4fd22, 0x2757d6, 0x8c54fe, 0x38c8b5, 0xccccc6, 0x505354, 0xff5995, 0xb6e354, 0xfeed6c, 0x3f78ff, 0x9e6ffe, 0x23cfd5, 0xf8f8f2)),
+    "the.hulk"            : (0xb5b5b5, 0x1b1d1e, (0x1b1d1e, 0x269d1b, 0x13ce30, 0x63e457, 0x2525f5, 0x641f74, 0x378ca9, 0xd9d8d1, 0x505354, 0x8dff2a, 0x48ff77, 0x3afe16, 0x506b95, 0x72589d, 0x4085a6, 0xe5e6e1)),
+    "tinacious.dark"      : (0xcbcbf0, 0x1d1d26, (0x1d1d26, 0xff3399, 0x00d364, 0xffcc66, 0x00cbff, 0xcc66ff, 0x00ceca, 0xcbcbf0, 0x636667, 0xff2f92, 0x00d364, 0xffd479, 0x00cbff, 0xd783ff, 0x00d5d4, 0xd5d6f3)),
+    "tinacious.light"     : (0x1d1d26, 0xf8f8ff, (0x1d1d26, 0xff3399, 0x00d364, 0xffcc66, 0x00cbff, 0xcc66ff, 0x00ceca, 0xcbcbf0, 0x636667, 0xff2f92, 0x00d364, 0xffd479, 0x00cbff, 0xd783ff, 0x00d5d4, 0xd5d6f3)),
+    "tokyonight"          : (0xc0caf5, 0x1a1b26, (0x15161e, 0xf7768e, 0x9ece6a, 0xe0af68, 0x7aa2f7, 0xbb9af7, 0x7dcfff, 0xa9b1d6, 0x414868, 0xf7768e, 0x9ece6a, 0xe0af68, 0x7aa2f7, 0xbb9af7, 0x7dcfff, 0xc0caf5)),
+    "tokyonight.day"      : (0x3760bf, 0xe1e2e7, (0xe9e9ed, 0xf52a65, 0x587539, 0x8c6c3e, 0x2e7de9, 0x9854f1, 0x007197, 0x6172b0, 0xa1a6c5, 0xf52a65, 0x587539, 0x8c6c3e, 0x2e7de9, 0x9854f1, 0x007197, 0x3760bf)),
+    "tokyonight.storm"    : (0xc0caf5, 0x24283b, (0x1d202f, 0xf7768e, 0x9ece6a, 0xe0af68, 0x7aa2f7, 0xbb9af7, 0x7dcfff, 0xa9b1d6, 0x414868, 0xf7768e, 0x9ece6a, 0xe0af68, 0x7aa2f7, 0xbb9af7, 0x7dcfff, 0xc0caf5)),
+    "tomorrow"            : (0x4d4d4c, 0xffffff, (0x000000, 0xc82829, 0x718c00, 0xeab700, 0x4271ae, 0x8959a8, 0x3e999f, 0xffffff, 0x000000, 0xc82829, 0x718c00, 0xeab700, 0x4271ae, 0x8959a8, 0x3e999f, 0xffffff)),
+    "tomorrow.night"      : (0xc5c8c6, 0x1d1f21, (0x000000, 0xcc6666, 0xb5bd68, 0xf0c674, 0x81a2be, 0xb294bb, 0x8abeb7, 0xffffff, 0x000000, 0xcc6666, 0xb5bd68, 0xf0c674, 0x81a2be, 0xb294bb, 0x8abeb7, 0xffffff)),
+    "tomorrow.night.1980" : (0xcccccc, 0x2d2d2d, (0x000000, 0xf2777a, 0x99cc99, 0xffcc66, 0x6699cc, 0xcc99cc, 0x66cccc, 0xffffff, 0x000000, 0xf2777a, 0x99cc99, 0xffcc66, 0x6699cc, 0xcc99cc, 0x66cccc, 0xffffff)),
+    "tomorrow.night.blue" : (0xffffff, 0x002451, (0x000000, 0xff9da4, 0xd1f1a9, 0xffeead, 0xbbdaff, 0xebbbff, 0x99ffff, 0xffffff, 0x000000, 0xff9da4, 0xd1f1a9, 0xffeead, 0xbbdaff, 0xebbbff, 0x99ffff, 0xffffff)),
+    "tomorrow.night.burn" : (0xa1b0b8, 0x151515, (0x252525, 0x832e31, 0xa63c40, 0xd3494e, 0xfc595f, 0xdf9395, 0xba8586, 0xf5f5f5, 0x5d6f71, 0x832e31, 0xa63c40, 0xd2494e, 0xfc595f, 0xdf9395, 0xba8586, 0xf5f5f5)),
+    "tomorrow.night.plus" : (0xeaeaea, 0x000000, (0x000000, 0xd54e53, 0xb9ca4a, 0xe7c547, 0x7aa6da, 0xc397d8, 0x70c0b1, 0xffffff, 0x000000, 0xd54e53, 0xb9ca4a, 0xe7c547, 0x7aa6da, 0xc397d8, 0x70c0b1, 0xffffff)),
+    "toychest"            : (0x31d07b, 0x24364b, (0x2c3f58, 0xbe2d26, 0x1a9172, 0xdb8e27, 0x325d96, 0x8a5edc, 0x35a08f, 0x23d183, 0x336889, 0xdd5944, 0x31d07b, 0xe7d84b, 0x34a6da, 0xae6bdc, 0x42c3ae, 0xd5d5d5)),
+    "treehouse"           : (0x786b53, 0x191919, (0x321300, 0xb2270e, 0x44a900, 0xaa820c, 0x58859a, 0x97363d, 0xb25a1e, 0x786b53, 0x433626, 0xed5d20, 0x55f238, 0xf2b732, 0x85cfed, 0xe14c5a, 0xf07d14, 0xffc800)),
+    "twilight"            : (0xffffd4, 0x141414, (0x141414, 0xc06d44, 0xafb97a, 0xc2a86c, 0x44474a, 0xb4be7c, 0x778385, 0xffffd4, 0x262626, 0xde7c4c, 0xccd88c, 0xe2c47e, 0x5a5e62, 0xd0dc8e, 0x8a989b, 0xffffd4)),
+    "ubuntu"              : (0xeeeeec, 0x300a24, (0x2e3436, 0xcc0000, 0x4e9a06, 0xc4a000, 0x3465a4, 0x75507b, 0x06989a, 0xd3d7cf, 0x555753, 0xef2929, 0x8ae234, 0xfce94f, 0x729fcf, 0xad7fa8, 0x34e2e2, 0xeeeeec)),
+    "ultra.dark"          : (0xffffff, 0x000000, (0x000000, 0xf07178, 0xc3e88d, 0xffcb6b, 0x82aaff, 0xc792ea, 0x89ddff, 0xcccccc, 0x333333, 0xf6a9ae, 0xdbf1ba, 0xffdfa6, 0xb4ccff, 0xddbdf2, 0xb8eaff, 0xffffff)),
+    "ultraviolent"        : (0xc1c1c1, 0x242728, (0x242728, 0xff0090, 0xb6ff00, 0xfff727, 0x47e0fb, 0xd731ff, 0x0effbb, 0xe1e1e1, 0x636667, 0xfb58b4, 0xdeff8c, 0xebe087, 0x7fecff, 0xe681ff, 0x69fcd3, 0xf9f9f5)),
+    "underthesea"         : (0xffffff, 0x011116, (0x022026, 0xb2302d, 0x00a941, 0x59819c, 0x459a86, 0x00599d, 0x5d7e19, 0x405555, 0x384451, 0xff4242, 0x2aea5e, 0x8ed4fd, 0x61d5ba, 0x1298ff, 0x98d028, 0x58fbd6)),
+    "unikitty"            : (0x0b0b0b, 0xff8cd9, (0x0c0c0c, 0xa80f20, 0xbafc8b, 0xeedf4b, 0x145fcd, 0xff36a2, 0x6bd1bc, 0xe2d7e1, 0x434343, 0xd91329, 0xd3ffaf, 0xffef50, 0x0075ea, 0xfdd5e5, 0x79ecd5, 0xfff3fe)),
+    "urple"               : (0x877a9b, 0x1b1b23, (0x000000, 0xb0425b, 0x37a415, 0xad5c42, 0x564d9b, 0x6c3ca1, 0x808080, 0x87799c, 0x5d3225, 0xff6388, 0x29e620, 0xf08161, 0x867aed, 0xa05eee, 0xeaeaea, 0xbfa3ff)),
+    "vaughn"              : (0xdcdccc, 0x25234f, (0x25234f, 0x705050, 0x60b48a, 0xdfaf8f, 0x5555ff, 0xf08cc3, 0x8cd0d3, 0x709080, 0x709080, 0xdca3a3, 0x60b48a, 0xf0dfaf, 0x5555ff, 0xec93d3, 0x93e0e3, 0xffffff)),
+    "vibrantink"          : (0xffffff, 0x000000, (0x878787, 0xff6600, 0xccff04, 0xffcc00, 0x44b4cc, 0x9933cc, 0x44b4cc, 0xf5f5f5, 0x555555, 0xff0000, 0x00ff00, 0xffff00, 0x0000ff, 0xff00ff, 0x00ffff, 0xe5e5e5)),
+    "vimbones"            : (0x353535, 0xf0f0ca, (0xf0f0ca, 0xa8334c, 0x4f6c31, 0x944927, 0x286486, 0x88507d, 0x3b8992, 0x353535, 0xc6c6a3, 0x94253e, 0x3f5a22, 0x803d1c, 0x1d5573, 0x7b3b70, 0x2b747c, 0x5c5c5c)),
+    "violet.dark"         : (0x708284, 0x1c1d1f, (0x56595c, 0xc94c22, 0x85981c, 0xb4881d, 0x2e8bce, 0xd13a82, 0x32a198, 0xc9c6bd, 0x45484b, 0xbd3613, 0x738a04, 0xa57705, 0x2176c7, 0xc61c6f, 0x259286, 0xc9c6bd)),
+    "violet.light"        : (0x536870, 0xfcf4dc, (0x56595c, 0xc94c22, 0x85981c, 0xb4881d, 0x2e8bce, 0xd13a82, 0x32a198, 0xd3d0c9, 0x45484b, 0xbd3613, 0x738a04, 0xa57705, 0x2176c7, 0xc61c6f, 0x259286, 0xc9c6bd)),
+    "warmneon"            : (0xafdab6, 0x404040, (0x000000, 0xe24346, 0x39b13a, 0xdae145, 0x4261c5, 0xf920fb, 0x2abbd4, 0xd0b8a3, 0xfefcfc, 0xe97071, 0x9cc090, 0xddda7a, 0x7b91d6, 0xf674ba, 0x5ed1e5, 0xd8c8bb)),
+    "wez"                 : (0xb3b3b3, 0x000000, (0x000000, 0xcc5555, 0x55cc55, 0xcdcd55, 0x5555cc, 0xcc55cc, 0x7acaca, 0xcccccc, 0x555555, 0xff5555, 0x55ff55, 0xffff55, 0x5555ff, 0xff55ff, 0x55ffff, 0xffffff)),
+    "whimsy"              : (0xb3b0d6, 0x29283b, (0x535178, 0xef6487, 0x5eca89, 0xfdd877, 0x65aef7, 0xaa7ff0, 0x43c1be, 0xffffff, 0x535178, 0xef6487, 0x5eca89, 0xfdd877, 0x65aef7, 0xaa7ff0, 0x43c1be, 0xffffff)),
+    "wildcherry"          : (0xdafaff, 0x1f1726, (0x000507, 0xd94085, 0x2ab250, 0xffd16f, 0x883cdc, 0xececec, 0xc1b8b7, 0xfff8de, 0x009cc9, 0xda6bac, 0xf4dca5, 0xeac066, 0x308cba, 0xae636b, 0xff919d, 0xe4838d)),
+    "wombat"              : (0xdedacf, 0x171717, (0x000000, 0xff615a, 0xb1e969, 0xebd99c, 0x5da9f6, 0xe86aff, 0x82fff7, 0xdedacf, 0x313131, 0xf58c80, 0xddf88f, 0xeee5b2, 0xa5c7ff, 0xddaaff, 0xb7fff9, 0xffffff)),
+    "wryan"               : (0x999993, 0x101010, (0x333333, 0x8c4665, 0x287373, 0x7c7c99, 0x395573, 0x5e468c, 0x31658c, 0x899ca1, 0x3d3d3d, 0xbf4d80, 0x53a6a6, 0x9e9ecb, 0x477ab3, 0x7e62b3, 0x6096bf, 0xc0c0c0)),
+    "zenbones"            : (0x2c363c, 0xf0edec, (0xf0edec, 0xa8334c, 0x4f6c31, 0x944927, 0x286486, 0x88507d, 0x3b8992, 0x2c363c, 0xcfc1ba, 0x94253e, 0x3f5a22, 0x803d1c, 0x1d5573, 0x7b3b70, 0x2b747c, 0x4f5e68)),
+    "zenbones.dark"       : (0xb4bdc3, 0x1c1917, (0x1c1917, 0xde6e7c, 0x819b69, 0xb77e64, 0x6099c0, 0xb279a7, 0x66a5ad, 0xb4bdc3, 0x403833, 0xe8838f, 0x8bae68, 0xd68c67, 0x61abda, 0xcf86c1, 0x65b8c1, 0x888f94)),
+    "zenbones.light"      : (0x2c363c, 0xf0edec, (0xf0edec, 0xa8334c, 0x4f6c31, 0x944927, 0x286486, 0x88507d, 0x3b8992, 0x2c363c, 0xcfc1ba, 0x94253e, 0x3f5a22, 0x803d1c, 0x1d5573, 0x7b3b70, 0x2b747c, 0x4f5e68)),
+    "zenburn"             : (0xdcdccc, 0x3f3f3f, (0x4d4d4d, 0x705050, 0x60b48a, 0xf0dfaf, 0x506070, 0xdc8cc3, 0x8cd0d3, 0xdcdccc, 0x709080, 0xdca3a3, 0xc3bf9f, 0xe0cf9f, 0x94bff3, 0xec93d3, 0x93e0e3, 0xffffff)),
+    "zenburned"           : (0xf0e4cf, 0x404040, (0x404040, 0xe3716e, 0x819b69, 0xb77e64, 0x6099c0, 0xb279a7, 0x66a5ad, 0xf0e4cf, 0x625a5b, 0xec8685, 0x8bae68, 0xd68c67, 0x61abda, 0xcf86c1, 0x65b8c1, 0xc0ab86)),
+    "zenwritten.dark"     : (0xbbbbbb, 0x191919, (0x191919, 0xde6e7c, 0x819b69, 0xb77e64, 0x6099c0, 0xb279a7, 0x66a5ad, 0xbbbbbb, 0x3d3839, 0xe8838f, 0x8bae68, 0xd68c67, 0x61abda, 0xcf86c1, 0x65b8c1, 0x8e8e8e)),
+    "zenwritten.light"    : (0x353535, 0xeeeeee, (0xeeeeee, 0xa8334c, 0x4f6c31, 0x944927, 0x286486, 0x88507d, 0x3b8992, 0x353535, 0xc6c3c3, 0x94253e, 0x3f5a22, 0x803d1c, 0x1d5573, 0x7b3b70, 0x2b747c, 0x5c5c5c)),
+}
diff --git a/Engine/Extras/ushell/channels/flow/core/cmds/tips.py b/Engine/Extras/ushell/channels/flow/core/cmds/tips.py
--- a/Engine/Extras/ushell/channels/flow/core/cmds/tips.py
+++ b/Engine/Extras/ushell/channels/flow/core/cmds/tips.py
@@ -33,9 +33,10 @@ class Tips(flow.cmd.Cmd):
             "Run '.help' to see all available commands",
             "Calvin and Hobbes is excellent",
             "Customise a session with '$USERPROFILE/.ushell/hooks/startup.bat'",
-            "'d:\>dc' changes directory using the fuzzy finder",
-            "'d:\>dc.' limits fuzzy-finder-chdir to the current directory",
-            "'d:\>dc DIR' limits fuzzy-finder-chdir to directory 'DIR'",
+            "'d:\\>dc' changes directory using the fuzzy finder",
+            "'d:\\>dc.' limits fuzzy-finder-chdir to the current directory",
+            "'d:\\>dc DIR' limits fuzzy-finder-chdir to directory 'DIR'",
             "Use Ctrl-P/Ctrl-N to recall previous queries to 'qq' and 'dc'",
             "'qq' to fuzzy find and execute lines from the command history",
+            "'.ushell theme' for setting and exploring terminal themes",
         )
diff --git a/Engine/Extras/ushell/channels/flow/core/describe.flow.py b/Engine/Extras/ushell/channels/flow/core/describe.flow.py
--- a/Engine/Extras/ushell/channels/flow/core/describe.flow.py
+++ b/Engine/Extras/ushell/channels/flow/core/describe.flow.py
@@ -24,42 +24,49 @@ debug_manifest.source("cmds/debug.py", "Debug")
 debug_manifest.invoke("debug", "manifest")
 debug_manifest.prefix("$")
 
-#-------------------------------------------------------------------------------
 debug_tools = flow.describe.Command()
 debug_tools.source("cmds/debug.py", "Tools")
 debug_tools.invoke("debug", "tools")
 debug_tools.prefix("$")
 
-#-------------------------------------------------------------------------------
 debug_sha1s = flow.describe.Command()
 debug_sha1s.source("cmds/debug.py", "Sha1s")
 debug_sha1s.invoke("debug", "tools", "sha1s")
 debug_sha1s.prefix("$")
 
-#-------------------------------------------------------------------------------
 debug_args = flow.describe.Command()
 debug_args.source("cmds/debug.py", "Argumentss")
 debug_args.invoke("debug", "args")
 debug_args.prefix("$")
 
-#-------------------------------------------------------------------------------
 debug_invalidate = flow.describe.Command()
 debug_invalidate.source("cmds/debug.py", "Invalidate")
 debug_invalidate.invoke("debug", "invalidate")
 debug_invalidate.prefix("$")
 
-#-------------------------------------------------------------------------------
 debug_paths = flow.describe.Command()
 debug_paths.source("cmds/debug.py", "Paths")
 debug_paths.invoke("debug", "paths")
 debug_paths.prefix("$")
 
-#-------------------------------------------------------------------------------
 debug_wipe = flow.describe.Command()
 debug_wipe.source("cmds/debug.py", "Wipe")
 debug_wipe.invoke("debug", "wipe")
 debug_wipe.prefix("$")
 
+#-------------------------------------------------------------------------------
+theme_set = flow.describe.Command()
+theme_set.source("cmds/theme.py", "Set")
+theme_set.invoke("ushell", "theme", "apply")
+
+theme_random = flow.describe.Command()
+theme_random.source("cmds/theme.py", "Random")
+theme_random.invoke("ushell", "theme", "random")
+
+theme_showcase = flow.describe.Command()
+theme_showcase.source("cmds/theme.py", "Showcase")
+theme_showcase.invoke("ushell", "theme", "showcase")
+
 #-------------------------------------------------------------------------------
 complete = flow.describe.Command()
 complete.source("cmds/complete.py", "Complete")
diff --git a/Engine/Extras/ushell/channels/flow/core/system/boot.py b/Engine/Extras/ushell/channels/flow/core/system/boot.py
--- a/Engine/Extras/ushell/channels/flow/core/system/boot.py
+++ b/Engine/Extras/ushell/channels/flow/core/system/boot.py
@@ -2,6 +2,7 @@
 
 import os
 import sys
+from pathlib import Path
 
 #-------------------------------------------------------------------------------
 def _enable_vt100():
@@ -33,11 +34,12 @@ def _main():
 
         collector.append(arg)
 
-    sys.path.append(os.path.abspath(__file__ + "/../lib"))
+    lib_dir = Path(__file__).resolve().parent / "lib"
+    sys.path.append(str(lib_dir))
 
     # Get all the channels ready and aggregated
     import bootstrap
-    working_dir, *channels_dirs = boot_dirs
+    working_dir, *channels_dirs = (Path(x) for x in boot_dirs)
     working_suffix = bootstrap.impl(working_dir, *channels_dirs)
 
     # Run the '$boot' command
diff --git a/Engine/Extras/ushell/channels/flow/core/system/flow/_flick.py b/Engine/Extras/ushell/channels/flow/core/system/flow/_flick.py
--- a/Engine/Extras/ushell/channels/flow/core/system/flow/_flick.py
+++ b/Engine/Extras/ushell/channels/flow/core/system/flow/_flick.py
@@ -299,12 +299,10 @@ class Cmd(object):
             double_dashed = False
             for arg in in_args:
                 double_dashed |= (arg == "--")
-                if arg in ("--help", "-h", "/?"):
-                    if not double_dashed:
-                        self._print_help()
-                        raise SystemExit(127)
-                else:
-                    yield arg
+                if not double_dashed and arg in ("--help", "-h", "/?"):
+                    self._print_help()
+                    raise SystemExit(127)
+                yield arg
 
         def on_error(message):
             self._print_help(short=True)
diff --git a/Engine/Extras/ushell/channels/flow/core/system/lib/bootstrap.py b/Engine/Extras/ushell/channels/flow/core/system/lib/bootstrap.py
--- a/Engine/Extras/ushell/channels/flow/core/system/lib/bootstrap.py
+++ b/Engine/Extras/ushell/channels/flow/core/system/lib/bootstrap.py
@@ -7,6 +7,7 @@ import hashlib
 import marshal
 import flow.describe
 import subprocess as sp
+from pathlib import Path
 
 #-------------------------------------------------------------------------------
 class _Log(object):
@@ -86,6 +87,10 @@ def _http_get(url, dest_dir, progress_cb=None):
 
 #-------------------------------------------------------------------------------
 def _extract_zip(payload, dest_dir):
+    # Use tar on POSIX platforms as Python's zipfile doesn't seem to apply attrs
+    if os.name != "nt":
+        return _extract_tar(payload, dest_dir)
+
     import zipfile
     zip_file = zipfile.ZipFile(payload, "r")
     zip_file.extractall(dest_dir)
@@ -571,9 +576,9 @@ class _Prune(_Action):
 
 
 #-------------------------------------------------------------------------------
-def impl(working_dir, *channels_dirs):
-    channels_dirs = [os.path.abspath(x) for x in channels_dirs if x]
-    working_dir = os.path.abspath(working_dir)
+def impl(working_dir:Path, *channels_dirs:Path):
+    channels_dirs = [str(x.resolve()) for x in channels_dirs if x.is_dir()]
+    working_dir = str(working_dir.resolve())
 
     def _read_channels(dir, depth=0):
         for item in fsutils.read_dirs(dir):
diff --git a/Engine/Extras/ushell/channels/flow/core/system/run.py b/Engine/Extras/ushell/channels/flow/core/system/run.py
--- a/Engine/Extras/ushell/channels/flow/core/system/run.py
+++ b/Engine/Extras/ushell/channels/flow/core/system/run.py
@@ -25,7 +25,7 @@ def _print_sub_commands_impl(tree_node):
             child_command = child_class.construct()
             desc = child_command.get_desc().lstrip()
         else:
-            desc = "[%s]" % ",".join(x for x,_ in child_node.read_children())
+            desc = "<%s>" % "|".join(x for x,_ in child_node.read_children())
         desc = desc.replace("\n", "")
         desc = desc.replace("    ", " ")
 
diff --git a/Engine/Extras/ushell/channels/flow/nt/describe.flow.py b/Engine/Extras/ushell/channels/flow/nt/describe.flow.py
--- a/Engine/Extras/ushell/channels/flow/nt/describe.flow.py
+++ b/Engine/Extras/ushell/channels/flow/nt/describe.flow.py
@@ -6,28 +6,28 @@ import flow.describe
 #-------------------------------------------------------------------------------
 clink = flow.describe.Tool()
 clink.version("1.0.0a4")
-clink.platform("win32")
 clink.payload("[REDACTED_PROFILE_URL]/clink/releases/download/1.0.0a4/clink-1.0.0a4.zip")
-clink.bin("clink_x64.exe")
 clink.sha1("6a642bc48e21a3093698067f0fe09eff3ccdfaa5")
+clink.platform("win32")
+clink.bin("clink_x64.exe")
 
 #-------------------------------------------------------------------------------
 fd = flow.describe.Tool()
-fd.version("8.4.0")
+fd.version("8.7.1")
+fd.payload("[REDACTED_PROFILE_URL]/fd/releases/download/v8.7.1/fd-v8.7.1-x86_64-pc-windows-msvc.zip")
+fd.sha1("7e8f5d0a9d1fed75b12483809f9619c04b39b60c")
 fd.platform("win32")
-fd.payload("[REDACTED_PROFILE_URL]/fd/releases/download/v8.4.0/fd-v8.4.0-x86_64-pc-windows-msvc.zip")
 fd.bin("fd.exe")
-fd.sha1("aef6f9809b7d86705fe828689453aae61f71365f")
 fd.source("[REDACTED_PROFILE_URL]/fd/releases/latest", "fd-v([0-9.]+)-x86_64")
 
 #-------------------------------------------------------------------------------
 fzf_win32 = flow.describe.Tool()
-fzf_win32.version("0.33.0")
+fzf_win32.version("0.44.1")
+fzf_win32.payload("[REDACTED_PROFILE_URL]/fzf/releases/download/0.44.1/fzf-0.44.1-windows_amd64.zip")
+fzf_win32.sha1("463f064d0403714ab5bf68ebf9fdea4c0159f6c3")
 fzf_win32.platform("win32")
-fzf_win32.payload("[REDACTED_PROFILE_URL]/fzf/releases/download/0.33.0/fzf-0.33.0-windows_amd64.zip")
 fzf_win32.bin("fzf.exe")
-fzf_win32.sha1("be7fc7b271069f7da5c27d570c81287d23cd328e")
-fzf_win32.source("[REDACTED_PROFILE_URL]/fzf/releases/latest", "fzf-(\d+\.\d+\.\d+)-windows")
+fzf_win32.source("[REDACTED_PROFILE_URL]/fzf/releases/latest", r"fzf-(\d+\.\d+\.\d+)-windows")
 
 #-------------------------------------------------------------------------------
 shell_cmd = flow.describe.Command()
diff --git a/Engine/Extras/ushell/channels/flow/nt/provision.bat b/Engine/Extras/ushell/channels/flow/nt/provision.bat
--- a/Engine/Extras/ushell/channels/flow/nt/provision.bat
+++ b/Engine/Extras/ushell/channels/flow/nt/provision.bat
@@ -11,9 +11,10 @@ call:check_bin tar.exe _tar_path
 call:check_bin curl.exe _curl_path
 call:check_bin certutil.exe _certutil_path
 
-set _pyver=3.11.4
-set _pysha=126802ff9fe787b961ae4d877262e6d6ce807d636295ef429c02e4dfd3e48041
-set _pytag=311
+set _pyver=3.12.1
+set _pysha=7a77c7471f5ceee88831a29b9460eddd784adf31b0230f188149ef0c588642d1
+set _pytag=312
+
 set _destdir=%~f1\%_pyver%
 
 if exist "%_destdir%" (
diff --git a/Engine/Extras/ushell/channels/flow/posix/boot.sh b/Engine/Extras/ushell/channels/flow/posix/boot.sh
--- a/Engine/Extras/ushell/channels/flow/posix/boot.sh
+++ b/Engine/Extras/ushell/channels/flow/posix/boot.sh
@@ -5,7 +5,7 @@
 self_dir=$(dirname $0)
 working_dir=~/.ushell/.working
 
-# provision python >=3.8 using provision.sh
+# provision Python using provision.sh
 if ! $self_dir/provision.sh $working_dir; then
     echo Failed to provision Python
     exit 1
diff --git a/Engine/Extras/ushell/channels/flow/posix/provision.sh b/Engine/Extras/ushell/channels/flow/posix/provision.sh
--- a/Engine/Extras/ushell/channels/flow/posix/provision.sh
+++ b/Engine/Extras/ushell/channels/flow/posix/provision.sh
@@ -6,54 +6,72 @@ _error()   { printf "\x1b[91m!! ERROR: %s\x1b[0m\n" "$1" ; }
 _header()  { printf "\x1b[96m-- %s\x1b[0m\n" "$1" ; }
 _success() { printf "\x1b[92m%s\x1b[0m\n" "$1" ; }
 
-py_ver=3.10
-
-# Try and find a suitable Python binary
-_header "Looking for a Python binary"
-if command -v python${py_ver}; then
-    py_bin=python${py_ver}
-elif command -v /usr/local/opt/python@${py_ver}/bin/python3; then
-    py_bin=/usr/local/opt/python@${py_ver}/bin/python3
-elif command -v python3; then
-    py_bin=python3
-else
-    echo FAILED
+py_ver_maj=3
+py_ver_min=12
+py_ver=$py_ver_maj.$py_ver_min
+
+# If we've already a destination we're done here
+if [ -z "$1" ]; then
+    _error "Missing 'working dir' argument"
     exit 1
 fi
-_success Found
 
-# Work out what version of Python we've got
-py_ver_full=$($py_bin --version)
-IFS=" "
-read -ra py_ver_full <<< "$py_ver_full"
-py_ver_full=${py_ver_full[1]}
+# If we've already a destination we're done here
+dest_dir=$1/python/$py_ver
+if [ -d "$dest_dir" ]; then
+    exit 0
+fi
 
 # Check the versions good enough
 function _check_version()
 {
     # Bash's here document syntax is ugly
-    cat << EOF | $1
-l = lambda x: [int(y) for y in x.split('.')][0:2]
-raise SystemExit(l('${py_ver_full}') < l('${py_ver}'))
+    cat << EOF | $1 -EsSB
+import sys
+vi = sys.version_info
+raise SystemExit(vi.major != ${py_ver_maj} or vi.minor < ${py_ver_min})
 EOF
     return $?
 }
 
-_check_version $py_bin
-if [ "$?" -ne "0" ]; then
-    _error "Version $py_ver_full is too old. $py_ver is required"
+function _bin_select()
+{
+    py_bin=
+    for candidate in "$@"; do
+        if ! command -v $candidate; then
+            continue
+        fi
+
+        if ! _check_version $candidate; then
+            continue
+        fi
+
+        py_bin=$candidate
+        return 0
+    done
+    return 1
+}
+
+_header "Looking for a $py_ver Python binary"
+_bin_select \
+    $USHELL_PY_BIN \
+    python$py_ver_maj \
+    /usr/local/opt/python@${py_ver_maj}.*/bin/python$py_ver_maj \
+    /usr/bin/python$py_ver_maj.*[0-9] \
+    /usr/local/bin/python$py_ver_maj.*[0-9] \
+    python
+
+if [ -z "$py_bin" ]; then
+    _error "No suitable binary found that is $py_ver or newer"
+    echo ""
     echo "Examples to install Python $py_ver:"
     echo "   Ubuntu: apt install python$py_ver"
     echo "      Mac: brew install python@$py_ver"
+    echo ""
     exit 1
 fi
-echo Using version $py_ver_full
 
-# Check if our destination exists already
-dest_dir=$1/python/$py_ver_full
-if [ -d "$dest_dir" ]; then
-    exit 0
-fi
+_success "Found: $py_bin: (as $($py_bin --version))"
 
 # Make sure that the Python install has Pip installed
 _header "Checking for Pip"
@@ -66,7 +84,7 @@ fi
 # Create the Python virtual env
 _header "Fetching virtual env"
 
-temp_dir=$1/python/${py_ver_full}_$$
+temp_dir=$1/python/${py_ver}_$$
 
 if ! $py_bin -m pip install virtualenv; then
     _error "Unable to install 'virtualenv' Pip package"
diff --git a/Engine/Extras/ushell/channels/unreal/core/cmds/build.py b/Engine/Extras/ushell/channels/unreal/core/cmds/build.py
--- a/Engine/Extras/ushell/channels/unreal/core/cmds/build.py
+++ b/Engine/Extras/ushell/channels/unreal/core/cmds/build.py
@@ -273,6 +273,7 @@ class _BuildCmd(unrealcmd.Cmd):
                         break
 
         builder.add_args("-SingleFile=" + file_or_module)
+        builder.add_args("-SkipDeploy")
 
     def is_constrained_build(self):
         return getattr(self, "_constrained_build", False)
@@ -344,8 +345,26 @@ class Build(_BuildCmd, unrealcmd.MultiPlatformCmd):
 
     def complete_target(self, prefix):
         ue_context = self.get_unreal_context()
-        yield from (x.name[:-10] for x in ue_context.glob("Source/*.Target.cs"))
-        yield from (x.name[:-10] for x in ue_context.glob("Source/Programs/**/*.Target.cs"))
+        ubt_manifests = list(ue_context.glob("Intermediate/Build/BuildRules/*RulesManifest.json"))
+        if not ubt_manifests:
+            yield from (x.name[:-10] for x in ue_context.glob("Source/*.Target.cs"))
+            yield from (x.name[:-10] for x in ue_context.glob("Source/Programs/*/*.Target.cs"))
+
+        for item in ubt_manifests:
+            with item.open("rt") as inp:
+                for line in inp:
+                    if ".Target.cs" not in line:
+                        continue
+
+                    slash = line.rfind("\\")
+                    if -1 == (slash := slash if slash != -1 else line.rfind("/")):
+                        continue
+
+                    line = line[slash + 1:]
+                    if -1 == (dot := line.find(".")):
+                        continue
+
+                    yield line[:dot]
 
     @unrealcmd.Cmd.summarise
     def main(self):
@@ -365,8 +384,33 @@ class Build(_BuildCmd, unrealcmd.MultiPlatformCmd):
         return self.run_build(builder)
 
 
+#-------------------------------------------------------------------------------
+def _add_clean_cmd(cmd_class):
+    # Create a class for a CleanX command
+    class _Clean(object):
+        def main(self):
+            self.args.clean = True
+            return super().main()
+
+    clean_name = "Clean" + cmd_class.__name__
+    clean_class = type(clean_name, (_Clean, cmd_class), {})
+    globals()[clean_name] = clean_class
+
+    # Rather crudely change the first work from Build to Clean
+    doc = cmd_class.__doc__.lstrip()
+    if doc.startswith("Builds"):
+        new_cmd_class = type(cmd_class.__name__, (cmd_class,), {})
+        new_cmd_class.__doc__ = doc
+        cmd_class.__doc__ = None
+        cmd_class = new_cmd_class
+
+        doc = "Cleans" + doc[6:]
+        clean_class.__doc__ = doc
+
+    return cmd_class
 
 #-------------------------------------------------------------------------------
+@_add_clean_cmd
 class Editor(_BuildCmd, unrealcmd.MultiPlatformCmd):
     """ Builds the editor. By default the current host platform is targeted. This
     can be overridden using the '--platform' option to cross-compile the editor."""
@@ -411,6 +455,7 @@ class Editor(_BuildCmd, unrealcmd.MultiPlatformCmd):
 
 
 #-------------------------------------------------------------------------------
+@_add_clean_cmd
 class Program(_BuildCmd, unrealcmd.MultiPlatformCmd):
     """ Builds a program for the current host platform (unless the --platform
     argument is provided) """
@@ -445,6 +490,7 @@ class Program(_BuildCmd, unrealcmd.MultiPlatformCmd):
 
 
 #-------------------------------------------------------------------------------
+@_add_clean_cmd
 class Server(_BuildCmd, unrealcmd.MultiPlatformCmd):
     """ Builds the server target for the host platform """
     platform = unrealcmd.Arg("", "Platform to build the server for")
@@ -496,6 +542,7 @@ class _Runtime(_BuildCmd, unrealcmd.MultiPlatformCmd):
         return self.run_build(builder)
 
 #-------------------------------------------------------------------------------
+@_add_clean_cmd
 class Client(_Runtime):
     """ Builds the client runtime """
     def main(self):
@@ -504,6 +551,7 @@ class Client(_Runtime):
         return super()._main_impl(target)
 
 #-------------------------------------------------------------------------------
+@_add_clean_cmd
 class Game(_Runtime):
     """ Builds the game runtime target binary executable """
     def main(self):
@@ -513,22 +561,6 @@ class Game(_Runtime):
 
 
 
-#-------------------------------------------------------------------------------
-class _CleanImpl(object):
-    def main(self):
-        self.args.clean = True
-        return super().main()
-
-def _make_clean_command(base):
-    return type("CleanImpl", (_CleanImpl, base), {})
-
-CleanEditor  = _make_clean_command(Editor)
-CleanProgram = _make_clean_command(Program)
-CleanServer  = _make_clean_command(Server)
-CleanClient  = _make_clean_command(Client)
-CleanGame    = _make_clean_command(Game)
-
-
 #-------------------------------------------------------------------------------
 _BuildCmd.complete_ubtargs = (
     "-2015", "-2017", "-2019", "-AllModules", "-alwaysgeneratedsym",
diff --git a/Engine/Extras/ushell/channels/unreal/core/cmds/cook.py b/Engine/Extras/ushell/channels/unreal/core/cmds/cook.py
--- a/Engine/Extras/ushell/channels/unreal/core/cmds/cook.py
+++ b/Engine/Extras/ushell/channels/unreal/core/cmds/cook.py
@@ -47,8 +47,8 @@ class _Cook(unrealcmd.MultiPlatformCmd):
         cultures = self.args.cultures.replace(",", "+")
 
         commandlet_args = (
-            "-targetplatform=" + cook_form,
-            "-cookcultures=" + cultures if cultures else None,
+            ("-targetplatform=" + cook_form) if cook_form else None,
+            ("-cookcultures=" + cultures) if cultures else None,
             "-unattended",
             "-unversioned",
             "-stdout",
diff --git a/Engine/Extras/ushell/channels/unreal/core/cmds/ddc.py b/Engine/Extras/ushell/channels/unreal/core/cmds/ddc.py
--- a/Engine/Extras/ushell/channels/unreal/core/cmds/ddc.py
+++ b/Engine/Extras/ushell/channels/unreal/core/cmds/ddc.py
@@ -23,6 +23,10 @@ class Auth(unrealcmd.Cmd):
         if value := default.OAuthProviderIdentifier:
             print("found;", value)
             self._service_name = value
+        else:
+            cloud = config.get("Engine", "StorageServers", "Cloud")
+            if value := cloud.OAuthProviderIdentifier:
+                self._service_name = value
 
     def main(self):
         ue_context = self.get_unreal_context()
diff --git a/Engine/Extras/ushell/channels/unreal/core/cmds/kill.py b/Engine/Extras/ushell/channels/unreal/core/cmds/kill.py
--- a/Engine/Extras/ushell/channels/unreal/core/cmds/kill.py
+++ b/Engine/Extras/ushell/channels/unreal/core/cmds/kill.py
@@ -6,6 +6,7 @@ import unrealcmd
 class Kill(unrealcmd.Cmd):
     """ Hard-terminates a running process on dev kits or the local machine. """
     what = unrealcmd.Arg(str, "Specified what platform or item to kill")
+    wait = unrealcmd.Opt(0.0, "Number of minutes to wait beforehand")
 
     def complete_what(self, prefix):
         yield "editor"
@@ -13,17 +14,30 @@ class Kill(unrealcmd.Cmd):
         yield "client"
         yield from self.complete_platform(prefix)
 
+    def _kill(self, platform, *args):
+        if self.args.wait > 0:
+            import time
+            until = time.time() + (self.args.wait * 60)
+            while True:
+                if time.time() >= until:
+                    break
+                print("Time remaining; %4ds" % max(0, int(until - time.time())), end="\r")
+                time.sleep(1.0)
+            print()
+
+        platform.kill(*args)
+
     def main(self):
         what = self.args.what
         try:
             platform = self.get_platform(what)
-            platform.kill()
+            self._kill(platform)
             return
         except ValueError:
             pass
 
         platform = self.get_platform()
-        if what == "editor":   platform.kill(what)
-        elif what == "server": platform.kill(what)
-        elif what == "client": platform.kill(what)
+        if what == "editor":   self._kill(platform, what)
+        elif what == "server": self._kill(platform, what)
+        elif what == "client": self._kill(platform, what)
         else: raise ValueError(f"I do not know how to '.kill {what}'")
diff --git a/Engine/Extras/ushell/channels/unreal/core/cmds/sln.py b/Engine/Extras/ushell/channels/unreal/core/cmds/sln.py
--- a/Engine/Extras/ushell/channels/unreal/core/cmds/sln.py
+++ b/Engine/Extras/ushell/channels/unreal/core/cmds/sln.py
@@ -5,22 +5,26 @@ import shutil
 import unrealcmd
 
 #-------------------------------------------------------------------------------
-class _Base(object):
+class _Base(unrealcmd.MultiPlatformCmd):
     def _get_primary_name(self):
-        name = "UE4"
-        try:
-            ue_context = self.get_unreal_context()
-            engine = ue_context.get_engine()
+        if hasattr(self, "_primary_name"):
+            return self._primary_name
 
+        ue_context = self.get_unreal_context()
+        engine = ue_context.get_engine()
+
+        name = "UE" + str(engine.get_version_major())
+        try:
             ppn_path = engine.get_dir() / "Intermediate/ProjectFiles/PrimaryProjectName.txt"
             if not ppn_path.is_file():
                 # Old path, needs to be maintained while projects prior to UE5.1 are supported
-                ppn_path = engine.get_dir() / "Intermediate/ProjectFiles/M"
-                ppn_path += "asterProjectName.txt"
+                ppn_path = engine.get_dir() / ("Intermediate/ProjectFiles/M" + "asterProjectName.txt")
             with ppn_path.open("rb") as ppn_file:
                 name = ppn_file.read().decode().strip()
         except FileNotFoundError:
-            pass
+            self.print_warning("Project files may be ungenerated")
+
+        self._primary_name = name
         return name
 
     def _get_sln_path(self):
@@ -34,29 +38,36 @@ class _Base(object):
             self.print_error("Opening project files is only supported on Windows")
             return False
 
+        self.print_info("Opening project")
         sln_path = self._get_sln_path()
+        print("Path:", sln_path)
+        if not os.path.isfile(sln_path):
+            self.print_error("Project file not found")
+            return False
 
-        self.print_info("Enumerating VS instances")
+        print("Enumerating VS instances")
         import vs.dte
         for i, instance in enumerate(vs.dte.running()):
-            if instance_path := instance.get_sln_path():
-                print(i, instance_path, end="")
-                if sln_path.samefile(instance_path):
-                    if instance.activate():
-                        print(" ...activating")
-                        return True
-                print()
+            print(f" {i} ", end="")
+            if not (instance_path := instance.get_sln_path()):
+                print("no-sln")
+                continue
+            print(instance_path, end="")
+
+            if sln_path.samefile(instance_path):
+                if instance.activate():
+                    print(" ...activating")
+                    return True
+            print()
 
         os.chdir(os.path.dirname(sln_path))
 
         run_args = ("cmd.exe", "/c", "start", sln_path)
         shell_open = self.get_exec_context().create_runnable(*run_args)
-        self.print_info("Opening project")
-        print("Path:", sln_path)
         return shell_open.run()
 
 #-------------------------------------------------------------------------------
-class Generate(unrealcmd.MultiPlatformCmd, _Base):
+class Generate(_Base):
     """ Generates a Visual Studio solution for the active project. To more
     easily distinguish open solutions from one branch to the next, the generated
     .sln file will be suffixed with the current branch name. Use the '--all'
@@ -98,7 +109,7 @@ class Generate(unrealcmd.MultiPlatformCmd, _Base):
             return self._open_sln()
 
 #-------------------------------------------------------------------------------
-class Open(unrealcmd.MultiPlatformCmd, _Base):
+class Open(_Base):
     """ Opens project files in Visual Studio """
 
     def get_exec_context(self):
diff --git a/Engine/Extras/ushell/channels/unreal/core/cmds/uat.py b/Engine/Extras/ushell/channels/unreal/core/cmds/uat.py
--- a/Engine/Extras/ushell/channels/unreal/core/cmds/uat.py
+++ b/Engine/Extras/ushell/channels/unreal/core/cmds/uat.py
@@ -6,8 +6,9 @@ import unrealcmd
 #-------------------------------------------------------------------------------
 class Uat(unrealcmd.MultiPlatformCmd):
     """ Runs an UnrealAutomationTool command """
-    command = unrealcmd.Arg(str, "The UnrealAutomationTool command to run")
-    uatargs = unrealcmd.Arg([str], "Arguments to pass to the UAT command")
+    command     = unrealcmd.Arg(str, "The UnrealAutomationTool command to run")
+    unprojected = unrealcmd.Opt(False, "No implicit '-project=' please")
+    uatargs     = unrealcmd.Arg([str], "Arguments to pass to the UAT command")
 
     def complete_command(self, prefix):
         import os
@@ -22,7 +23,7 @@ class Uat(unrealcmd.MultiPlatformCmd):
                 if entry.name.endswith(".cs"):
                     with open(entry.path, "rt") as lines:
                         for line in lines:
-                            m = re.search("(\w+) : BuildCommand", line)
+                            m = re.search(r"(\w+) : BuildCommand", line)
                             if m:
                                 yield m.group(1)
                                 break
@@ -60,8 +61,9 @@ class Uat(unrealcmd.MultiPlatformCmd):
         cmd /= "RunUAT.bat" if os.name == "nt" else "RunUAT.sh"
         args = (self.args.command,)
 
-        if project := ue_context.get_project():
-            args = (*args, "-project=" + str(project.get_path()))
+        if args[0] and not self.args.unprojected:
+            if project := ue_context.get_project():
+                args = (*args, "-project=" + str(project.get_path()))
 
         exec_context = self.get_exec_context()
         cmd = exec_context.create_runnable(cmd, *args, *self.args.uatargs)
diff --git a/Engine/Extras/ushell/channels/unreal/core/cmds/zen.py b/Engine/Extras/ushell/channels/unreal/core/cmds/zen.py
new file mode 100644
--- /dev/null
+++ b/Engine/Extras/ushell/channels/unreal/core/cmds/zen.py
@@ -0,0 +1,323 @@
+# Copyright Epic Games, Inc. All Rights Reserved.
+
+import os
+import sys
+import shutil
+import unreal
+import unrealcmd
+import unreal.cmdline
+import http.client
+import json
+import hashlib
+
+#-------------------------------------------------------------------------------
+class _Base(unrealcmd.Cmd):
+    def perform_project_completion(self, prefix, hostname="localhost", port=8558):
+        headers = {"Accept":"application/json"}
+        conn = http.client.HTTPConnection(hostname, port=port, timeout=0.1)
+        projects = []
+        try:
+            conn.request("GET", "/prj", headers=headers)
+            response = conn.getresponse()
+            projects = json.loads(response.read().decode())
+        except socket.timeout:
+            return
+
+        for project in projects:
+            yield project['Id']
+
+    def perform_oplog_completion(self, prefix, project_filter, hostname="localhost", port=8558):
+        headers = {"Accept":"application/json"}
+        conn = http.client.HTTPConnection(hostname, port=port, timeout=0.1)
+        projects = []
+        try:
+            if project_filter:
+                conn.request("GET", f"/prj/{project_filter}", headers=headers)
+                response = conn.getresponse()
+                projects.append(json.loads(response.read().decode()))
+            else:
+                conn.request("GET", "/prj", headers=headers)
+                response = conn.getresponse()
+                projects = json.loads(response.read().decode())
+        except socket.timeout:
+            return
+
+        for project in projects:
+            for oplog in project['oplogs']:
+                yield oplog['id']
+
+    def get_project_id_from_project(self, project):
+        if not project:
+            self.print_error("An active project is required for this operation")
+
+        normalized_project_path = str(project.get_path()).replace("\\","/")
+        project_id = project.get_name() + "." + hashlib.md5(normalized_project_path.encode('utf-8')).hexdigest()[:8]
+        return project_id;
+
+
+    def get_project_or_default(self, explicit_project):
+        if explicit_project:
+            return explicit_project
+
+        ue_context = self.get_unreal_context()
+        return self.get_project_id_from_project(ue_context.get_project())
+
+
+#-------------------------------------------------------------------------------
+class _ZenUETargetBase(_Base):
+    variant = unrealcmd.Arg("development", "Build variant to launch")
+    runargs  = unrealcmd.Arg([str], "Additional arguments to pass to the process being launched")
+    build    = unrealcmd.Opt(True, "Build the target prior to running it")
+
+    def _build(self, target, variant, platform):
+        if target.get_type() == unreal.TargetType.EDITOR:
+            build_cmd = ("_build", "editor", variant)
+        else:
+            platform = platform or unreal.Platform.get_host()
+            build_cmd = ("_build", "target", target.get_name(), platform, variant)
+
+        import subprocess
+        ret = subprocess.run(build_cmd)
+        if ret.returncode:
+            raise SystemExit(ret.returncode)
+
+    def get_binary_path(self, target, platform=None):
+        ue_context = self.get_unreal_context()
+
+        if isinstance(target, str):
+            target = ue_context.get_target_by_name(target)
+        else:
+            target = ue_context.get_target_by_type(target)
+
+        if self.args.build:
+            self._build(target, self.args.variant, platform)
+
+        variant = unreal.Variant.parse(self.args.variant)
+        if build := target.get_build(variant, platform):
+            if build := build.get_binary_path():
+                return build
+
+        raise EnvironmentError(f"No {self.args.variant} build found for target '{target.get_name()}'")
+
+    def run_ue_program_target(self, target_name, cmdargs = None):
+        binary_path = self.get_binary_path(target_name)
+
+        final_args = []
+        if cmdargs: final_args += cmdargs
+        if self.args.runargs: final_args += (*unreal.cmdline.read_ueified(*self.args.runargs),)
+
+        launch_kwargs = {}
+        if not sys.stdout.isatty():
+            launch_kwargs = { "silent" : True }
+        exec_context = self.get_exec_context()
+        runnable = exec_context.create_runnable(str(binary_path), *final_args);
+        runnable.launch(**launch_kwargs)
+        return True if runnable.is_gui() else runnable.wait()
+
+
+#-------------------------------------------------------------------------------
+class _ZenUtilityBase(_Base):
+    zenargs  = unrealcmd.Arg([str], "Additional arguments to pass to zen utility")
+
+    def get_zen_utility_command(self, cmdargs):
+        ue_context = self.get_unreal_context()
+        engine = ue_context.get_engine()
+
+        platform = unreal.Platform.get_host()
+        exe_filename = "zen"
+        if platform == "Win64":
+            exe_filename = exe_filename + ".exe"
+
+        bin_path = engine.get_dir()
+        bin_path /= f"Binaries/{platform}/{exe_filename}"
+        if not bin_path.is_file():
+            raise FileNotFoundError(f"Unable to find '{bin_path}'")
+
+        final_args = []
+        if cmdargs: final_args += cmdargs
+        if self.args.zenargs: final_args += self.args.zenargs
+
+        return bin_path, final_args
+
+
+    def run_zen_utility(self, cmdargs):
+        bin_path, final_args = self.get_zen_utility_command(cmdargs)
+
+        cmd = self.get_exec_context().create_runnable(bin_path, *final_args)
+        return cmd.run()
+
+#-------------------------------------------------------------------------------
+class Dashboard(_ZenUETargetBase):
+    """ Starts the zen dashboard GUI."""
+
+    def main(self):
+        return self.run_ue_program_target("ZenDashboard")
+
+#-------------------------------------------------------------------------------
+class Start(_ZenUETargetBase):
+    """ Starts an instance of zenserver."""
+    SponsorProcessID = unrealcmd.Opt("", "Process ID to be added as a sponsor process for the zenserver process")
+
+    def main(self):
+
+        args = []
+        if self.args.SponsorProcessID:
+            args.append("-SponsorProcessID=" + str(self.args.SponsorProcessID))
+        else:
+            args.append("-SponsorProcessID=" + str(os.getppid()))
+
+        return self.run_ue_program_target("ZenLaunch", args)
+
+#-------------------------------------------------------------------------------
+class Stop(_ZenUtilityBase):
+    """ Stops any running instance of zenserver."""
+
+    def main(self):
+        return self.run_zen_utility(["down"])
+
+#-------------------------------------------------------------------------------
+class Status(_ZenUtilityBase):
+    """ Get the status of running zenserver instances."""
+
+    def main(self):
+        return self.run_zen_utility(["status"])
+
+#-------------------------------------------------------------------------------
+class Version(_ZenUtilityBase):
+    """ Get the version of the in-tree zenserver executable."""
+
+    def main(self):
+        return self.run_zen_utility(["version"])
+
+#-------------------------------------------------------------------------------
+class ImportSnapshot(_ZenUtilityBase):
+    """ Imports an oplog snapshot into the running zenserver process."""
+    snapshotdescriptor = unrealcmd.Arg(str, "Snapshot descriptor file path to import from")
+    snapshotindex = unrealcmd.Arg(0, "0-based index of the snapshot within the snapshot descriptor to import from")
+    projectid = unrealcmd.Opt("", "The zen project ID to import into (defaults to an ID based on the current ushell project)")
+    oplog = unrealcmd.Opt("", "The zen oplog to import into (defaults to the oplog name in the snapshot)")
+    cloudauthservice = unrealcmd.Opt("", "Name of the service to authorize with when importing from a cloud source")
+    asyncimport = unrealcmd.Opt(False, "Trigger import but don't wait for completion")
+
+    def complete_projectid(self, prefix):
+        return self.perform_project_completion(prefix)
+
+    def complete_oplog(self, prefix):
+        return self.perform_oplog_completion(prefix, self.get_project_or_default(self.args.projectid))
+
+    def complete_snapshotdescriptor(self, prefix):
+        prefix = prefix or "."
+        for item in os.scandir(prefix):
+            if item.name.endswith(".json"):
+                yield item.name
+                return
+
+        for item in os.scandir(prefix):
+            if item.is_dir():
+                yield item.name + "/"
+
+    def _lookup_service_name(self):
+        if getattr(self, "_service_name", None):
+            return
+
+        config = self.get_unreal_context().get_config()
+        default = config.get("Engine", "StorageServers", "Default")
+        if value := default.OAuthProviderIdentifier:
+            self._service_name = value
+        else:
+            cloud = config.get("Engine", "StorageServers", "Cloud")
+            if value := cloud.OAuthProviderIdentifier:
+                self._service_name = value
+
+    def refresh_zen_token(self):
+        ue_context = self.get_unreal_context()
+        engine = ue_context.get_engine()
+
+        bin_type = "win-x64"
+        if sys.platform == "darwin": bin_type = "osx-x64"
+        elif sys.platform == "linux": bin_type = "linux-x64"
+
+        bin_path = engine.get_dir()
+        bin_path /= f"Binaries/DotNET/OidcToken/{bin_type}/OidcToken.exe"
+        if not bin_path.is_file():
+            raise FileNotFoundError(f"Unable to find '{bin_path}'")
+
+        self._service_name = self.args.cloudauthservice
+        self._lookup_service_name()
+        if not self._service_name:
+            raise ValueError("Unable to discover service name")
+
+        oidcargs = (
+            "--ResultToConsole=true",
+            "--Service=" + str(self._service_name),
+        )
+
+        if project := ue_context.get_project():
+            oidcargs = (*oidcargs, "--Project=" + str(project.get_path()))
+
+        cmd = self.get_exec_context().create_runnable(bin_path, *oidcargs)
+        ret, output = cmd.run2()
+        if ret != 0:
+            return False
+
+        tokenresponse = json.loads(output)
+
+        self._AccessToken = tokenresponse['Token']
+        self.print_info("Cloud access token obtained for import operation.")
+
+        return True
+
+    def get_exec_context(self):
+        context = super().get_exec_context()
+        if hasattr(self, '_AccessToken') and self._AccessToken:
+            context.get_env()["UE-CloudDataCacheAccessToken"] = self._AccessToken
+        return context
+
+    def add_args_from_descriptor(self, snapshot, args):
+        snapshot_type = snapshot['type']
+
+        # TODO: file, zen snapshot types
+        if snapshot_type == 'cloud':
+            args.append('--cloud')
+            args.append(snapshot['host'])
+            args.append('--namespace')
+            args.append(snapshot['namespace'])
+            args.append('--bucket')
+            args.append(snapshot['bucket'])
+            args.append('--key')
+            args.append(snapshot['key'])
+        else:
+            self.print_error(f"Unsupported snapshot type {snapshot_type}")
+            return False
+
+        return True
+
+    @unrealcmd.Cmd.summarise
+    def main(self):
+        try:
+            with open(self.args.snapshotdescriptor, "rt") as file:
+                descriptor = json.load(file)
+        except FileNotFoundError:
+            self.print_error(f"Error accessing snapshot descriptor file {self.args.snapshotdescriptor}")
+            return False
+
+        snapshot = descriptor["snapshots"][self.args.snapshotindex]
+        snapshot_type = snapshot['type']
+
+        args = [
+        "oplog-import",
+        self.get_project_or_default(self.args.projectid),
+        self.args.oplog or snapshot['targetplatform'],
+        ]
+
+        if self.args.asyncimport:
+            args.append("--async")
+
+        if not self.add_args_from_descriptor(snapshot, args):
+            return False
+
+        if snapshot_type == 'cloud':
+            if not self.refresh_zen_token():
+                return False
+
+        return self.run_zen_utility(args)
diff --git a/Engine/Extras/ushell/channels/unreal/core/describe.flow.py b/Engine/Extras/ushell/channels/unreal/core/describe.flow.py
--- a/Engine/Extras/ushell/channels/unreal/core/describe.flow.py
+++ b/Engine/Extras/ushell/channels/unreal/core/describe.flow.py
@@ -5,43 +5,43 @@ import flow.describe
 
 #-------------------------------------------------------------------------------
 fzf_linux = flow.describe.Tool()
-fzf_linux.version("0.33.0")
+fzf_linux.version("0.44.1")
+fzf_linux.payload("[REDACTED_PROFILE_URL]/fzf/releases/download/0.44.1/fzf-0.44.1-linux_amd64.tar.gz")
+fzf_linux.sha1("e7780f1e5e2dd4c8aa5e94dcbcbedf4cebfaaceb")
 fzf_linux.platform("linux")
-fzf_linux.payload("[REDACTED_PROFILE_URL]/fzf/releases/download/0.33.0/fzf-0.33.0-linux_amd64.tar.gz")
 fzf_linux.bin("fzf")
-fzf_linux.sha1("41aeb1c7bc01c700614768a777c5df91919a180c")
 
 fzf_darwin = flow.describe.Tool()
-fzf_darwin.version("0.33.0")
+fzf_darwin.version("0.44.1")
+fzf_darwin.payload("[REDACTED_PROFILE_URL]/fzf/releases/download/0.44.1/fzf-0.44.1-darwin_amd64.zip")
+fzf_darwin.sha1("405156eb6fc3d5275774cc569e24feaed362eab5")
 fzf_darwin.platform("darwin")
-fzf_darwin.payload("[REDACTED_PROFILE_URL]/fzf/releases/download/0.33.0/fzf-0.33.0-darwin_amd64.zip")
 fzf_darwin.bin("fzf")
-fzf_darwin.sha1("f83505967ca3d0ad849a461029d60cd414e3ed76")
 
 
 
 #-------------------------------------------------------------------------------
 ripgrep_win32 = flow.describe.Tool()
-ripgrep_win32.version("13.0.0")
+ripgrep_win32.version("14.0.3")
+ripgrep_win32.payload("[REDACTED_PROFILE_URL]/ripgrep/releases/download/14.0.3/ripgrep-14.0.3-x86_64-pc-windows-msvc.zip")
+ripgrep_win32.sha1("5508b3dd5c12206c5d1c7994490bacf317cee0aa")
 ripgrep_win32.platform("win32")
-ripgrep_win32.payload("[REDACTED_PROFILE_URL]/ripgrep/releases/download/13.0.0/ripgrep-13.0.0-x86_64-pc-windows-msvc.zip")
 ripgrep_win32.bin("rg.exe")
-ripgrep_win32.sha1("f54ffef82c50d6079f211ed5348a81aa98274881")
-ripgrep_win32.source("[REDACTED_PROFILE_URL]/ripgrep/releases/latest", "ripgrep-(\d+\.\d+\.\d+)-x86_64")
+ripgrep_win32.source("[REDACTED_PROFILE_URL]/ripgrep/releases/latest", r"ripgrep-(\d+\.\d+\.\d+)-x86_64")
 
 ripgrep_linux = flow.describe.Tool()
-ripgrep_linux.version("13.0.0")
+ripgrep_linux.version("14.0.3")
+ripgrep_linux.payload("[REDACTED_PROFILE_URL]/ripgrep/releases/download/14.0.3/ripgrep-14.0.3-x86_64-unknown-linux-musl.tar.gz")
+ripgrep_linux.sha1("25a759834819f91625ac739c890a4b0139df7b63")
 ripgrep_linux.platform("linux")
-ripgrep_linux.payload("[REDACTED_PROFILE_URL]/ripgrep/releases/download/13.0.0/ripgrep-13.0.0-x86_64-unknown-linux-musl.tar.gz")
 ripgrep_linux.bin("rg")
-ripgrep_linux.sha1("eae550a5ed0fcef8bc555c9ec41c593b8b5aea46")
 
 ripgrep_darwin = flow.describe.Tool()
-ripgrep_darwin.version("13.0.0")
+ripgrep_darwin.version("14.0.3")
+ripgrep_darwin.payload("[REDACTED_PROFILE_URL]/ripgrep/releases/download/14.0.3/ripgrep-14.0.3-x86_64-apple-darwin.tar.gz")
+ripgrep_darwin.sha1("d5af393ea9b5a443544edc686cf939f517be35ea")
 ripgrep_darwin.platform("darwin")
-ripgrep_darwin.payload("[REDACTED_PROFILE_URL]/ripgrep/releases/download/13.0.0/ripgrep-13.0.0-x86_64-apple-darwin.tar.gz")
 ripgrep_darwin.bin("rg")
-ripgrep_darwin.sha1("f09500e866c49c641ecb42fd8a0089d664d7d385")
 
 
 
@@ -146,6 +146,31 @@ ddc_auth = flow.describe.Command()
 ddc_auth.source("cmds/ddc.py", "Auth")
 ddc_auth.invoke("ddc", "auth")
 
+#-------------------------------------------------------------------------------
+zen_start = flow.describe.Command()
+zen_start.source("cmds/zen.py", "Start")
+zen_start.invoke("zen", "start")
+
+zen_stop = flow.describe.Command()
+zen_stop.source("cmds/zen.py", "Stop")
+zen_stop.invoke("zen", "stop")
+
+zen_dashboard = flow.describe.Command()
+zen_dashboard.source("cmds/zen.py", "Dashboard")
+zen_dashboard.invoke("zen", "dashboard")
+
+zen_status = flow.describe.Command()
+zen_status.source("cmds/zen.py", "Status")
+zen_status.invoke("zen", "status")
+
+zen_version = flow.describe.Command()
+zen_version.source("cmds/zen.py", "Version")
+zen_version.invoke("zen", "version")
+
+zen_importsnapshot = flow.describe.Command()
+zen_importsnapshot.source("cmds/zen.py", "ImportSnapshot")
+zen_importsnapshot.invoke("zen", "importsnapshot")
+
 #-------------------------------------------------------------------------------
 sln_generate = flow.describe.Command()
 sln_generate.source("cmds/sln.py", "Generate")
@@ -218,11 +243,11 @@ tips.prefix("$")
 #-------------------------------------------------------------------------------
 vswhere = flow.describe.Tool()
 vswhere.version("3.0.1")
+vswhere.payload("[REDACTED_PROFILE_URL]/vswhere/releases/download/3.1.7/vswhere.exe")
+vswhere.sha1("e3fa9b2db259d8875170717469779ea1280c8466")
 vswhere.platform("win32")
-vswhere.payload("[REDACTED_PROFILE_URL]/vswhere/releases/download/3.0.1/vswhere.exe")
 vswhere.bin("vswhere.exe")
-vswhere.sha1("209d10bd07c6dc032ad26b40e4d93227c7c385fb")
-vswhere.source("[REDACTED_PROFILE_URL]/vswhere/releases", "(\d+\.\d+\.\d+)/vswhere.exe")
+vswhere.source("[REDACTED_PROFILE_URL]/vswhere/releases", r"(\d+\.\d+\.\d+)/vswhere.exe")
 
 
 
diff --git a/Engine/Extras/ushell/channels/unreal/core/pylib/unreal/_context.py b/Engine/Extras/ushell/channels/unreal/core/pylib/unreal/_context.py
--- a/Engine/Extras/ushell/channels/unreal/core/pylib/unreal/_context.py
+++ b/Engine/Extras/ushell/channels/unreal/core/pylib/unreal/_context.py
@@ -2,9 +2,9 @@
 
 import os
 import re
-import pathlib
 from enum import Enum, Flag
-from typing import ForwardRef, Iterator, Optional, Tuple, Dict
+from typing import Iterator
+from pathlib import Path, WindowsPath, PosixPath
 
 from . import _ini as ini
 from ._ubt import BuildTool
@@ -13,7 +13,7 @@ from ._platform import Platform
 # {{{1 misc --------------------------------------------------------------------
 
 #-------------------------------------------------------------------------------
-def _resolve_association_nt(association:str) -> Optional[pathlib.Path]:
+def _resolve_association_nt(association:str) -> Path|None:
     import winreg
 
     if association[0] == "{":
@@ -32,12 +32,12 @@ def _resolve_association_nt(association:str) -> Optional[pathlib.Path]:
         return
 
     if value[1] == 1:
-        return pathlib.Path(value[0])
+        return Path(value[0])
 
-def _resolve_association_posix(association:str) -> Optional[pathlib.Path]:
+def _resolve_association_posix(association:str) -> Path|None:
     raise NotImplementedError("Foreign projects are not yet supported")
 
-def _resolve_association(association:str) -> Optional[pathlib.Path]:
+def _resolve_association(association:str) -> Path|None:
     if os.name == "nt":
         return _resolve_association_nt(association)
     else:
@@ -59,7 +59,7 @@ class Variant(Enum):
         return "DebugGame" if self == Variant.DEBUGGAME else self.name.title()
 
     @staticmethod
-    def parse(external:str) -> ForwardRef("Variant"):
+    def parse(external:str) -> "Variant":
         try: return Variant[external.upper()]
         except: raise ValueError(f"Unknown build variant '{external}'")
 
@@ -73,7 +73,7 @@ class TargetType(Enum):
     SERVER      = 0x05
 
     @staticmethod
-    def parse(external:str) -> ForwardRef("TargetType"):
+    def parse(external:str) -> "TargetType":
         try: return TargetType[external.upper()]
         except: raise ValueError(f"Unknown target type '{external}'")
 
@@ -86,46 +86,49 @@ class ContextType(Enum):
 
 
 # {{{1 ufs-item-----------------------------------------------------------------
-class _UfsItem(pathlib.Path):
-    _mount:ForwardRef("_UfsMount")
+class _UfsItem(Path):
+    _mount:"_UfsMount"
 
-    def __new__(cls, mount:ForwardRef("_UfsMount"), *args) -> ForwardRef("_UfsItem"):
+    def __new__(cls, mount:"_UfsMount", *args) -> "_UfsItem":
         cls = _WindowsUfsItem if os.name == "nt" else _PosixUfsItem
         instance = super().__new__(cls, *args)
         instance._mount = mount
         return instance
 
-    def get_mount(self) -> ForwardRef("_UfsMount"):
+    def get_mount(self) -> "_UfsMount":
         return self._mount
 
-class _WindowsUfsItem(pathlib.WindowsPath, _UfsItem):
+class _WindowsUfsItem(WindowsPath, _UfsItem):
     pass
 
-class _PosixUfsItem(pathlib.PosixPath, _UfsItem):
+class _PosixUfsItem(PosixPath, _UfsItem):
     pass
 
 
 
 # {{{1 ufs_mount ---------------------------------------------------------------
 class _UfsMount(object):
-    _dir    :pathlib.Path
-    _node   :ForwardRef("_UfsNode")
+    _dir    :Path
+    _node   :"_UfsNode"
 
-    def __init__(self, dir:pathlib.Path, node:ForwardRef("_UfsNode")):
+    def __init__(self, dir:Path, node:"_UfsNode"):
         self._dir = dir
         self._node = node
 
-    def get_dir(self) -> pathlib.Path:
+    def __fspath__(self):
+        return str(self._dir)
+
+    def get_dir(self) -> Path:
         return self._dir
 
-    def get_node(self) -> ForwardRef("_UfsNode"):
+    def get_node(self) -> "_UfsNode":
         return self._node
 
     def glob(self, pattern:str) -> Iterator[_UfsItem]:
         for x in self.get_dir().glob(pattern):
             yield _UfsItem(self, x)
 
-    def get_item(self, file_path:str) -> Optional[_UfsItem]:
+    def get_item(self, file_path:str) -> _UfsItem|None:
         dest = self.get_dir() / file_path
         if dest.is_file():
             return _UfsItem(self, dest)
@@ -134,12 +137,12 @@ class _UfsMount(object):
 
 # {{{1 build -------------------------------------------------------------------
 class _Build(object):
-    _receipt_path   :pathlib.Path
+    _receipt_path   :Path
     _variant        :Variant
     _ufs_mount      :_UfsMount
     _platform       :str
 
-    def __init__(self, receipt_path:pathlib.Path, variant:Variant, platform:str, ufs_mount:_UfsMount):
+    def __init__(self, receipt_path:Path, variant:Variant, platform:str, ufs_mount:_UfsMount):
         self._receipt_path = receipt_path
         self._variant = variant
         self._platform = platform
@@ -151,7 +154,7 @@ class _Build(object):
     def get_platform(self) -> str:
         return self._platform
 
-    def get_binary_path(self) -> Optional[pathlib.Path]:
+    def get_binary_path(self) -> Path|None:
         with self._receipt_path.open("rt") as x:
             try:
                 import json
@@ -170,7 +173,7 @@ class _Build(object):
         binary_path = binary_path.replace("$(ProjectDir)", str(proj_node.get_dir()))
         binary_path = binary_path.replace("$(EngineDir)", str(eng_node.get_dir()))
 
-        ret = pathlib.Path(binary_path)
+        ret = Path(binary_path)
         return ret if ret.exists() else None
 
 
@@ -195,7 +198,7 @@ class _Target(object):
     def is_project_target(self) -> bool:
         return isinstance(self._ufs_mount.get_node(), _Project)
 
-    def get_build(self, variant:Variant=Variant.DEVELOPMENT, platform:str="") -> Optional[_Build]:
+    def get_build(self, variant:Variant=Variant.DEVELOPMENT, platform:str="") -> _Build|None:
         platform = platform or Platform.get_host()
 
         path = f"Binaries/{platform}/{self.get_name()}"
@@ -210,17 +213,17 @@ class _Target(object):
 
 # {{{1 ufs_node ----------------------------------------------------------------
 class _UfsNode(object):
-    _parent :ForwardRef("_UfsNode")
-    _dir    :pathlib.Path
+    _parent :"_UfsNode"
+    _dir    :Path
 
-    def __init__(self, dir:pathlib.Path, parent:ForwardRef("_UfsNode")=None):
+    def __init__(self, dir:Path, parent:"_UfsNode"=None):
         self._parent = parent
         self._dir = dir
 
-    def get_parent_node(self) -> Optional[ForwardRef("_UfsNode")]:
+    def get_parent_node(self) -> "_UfsNode|None":
         return self._parent
 
-    def get_dir(self) -> pathlib.Path:
+    def get_dir(self) -> Path:
         return self._dir
 
     def get_root_mount(self) -> _UfsMount:
@@ -259,14 +262,14 @@ class _UfsNode(object):
                     type = find_line(r"=\s*TargetType\.(.+)\s*;", line_iter)
                     return name, type or "UNKNOWN"
 
-        pattern = str(pathlib.Path(targets_dir) / "*.Target.cs")
+        pattern = str(Path(targets_dir) / "*.Target.cs")
         for target_cs in self.glob(pattern, **kwargs):
             if info := parse_target_cs(target_cs):
                 name, type = info
                 type = TargetType.parse(type)
                 yield _Target(name, type, target_cs.get_mount())
 
-    def get_target_by_type(self, target_type:TargetType) -> Optional[_Target]:
+    def get_target_by_type(self, target_type:TargetType) -> _Target|None:
         if target_type.value <= TargetType.PROGRAM.value:
             return
 
@@ -274,7 +277,7 @@ class _UfsNode(object):
             if target.get_type() == target_type:
                 return target
 
-    def get_target_by_name(self, target_name:str) -> Optional[_Target]:
+    def get_target_by_name(self, target_name:str) -> _Target|None:
         lower_name = target_name.lower()
 
         for dir in ("Source/Programs/*", "Source"):
@@ -326,7 +329,7 @@ class _Config(object):
 
         return ret
 
-    def get(self, category:str, section:str="", key:str="", *, override:str="") -> Optional[ini.Ini|ini._Struct|ini._Value]:
+    def get(self, category:str, section:str="", key:str="", *, override:str="") -> ini.Ini|ini._Struct|ini._Value|None:
         ret = self._load_hives(category, override=override)
 
         if section:
@@ -339,7 +342,7 @@ class _Config(object):
 
 # {{{1 engine ------------------------------------------------------------------
 class _Engine(_UfsNode):
-    def __init__(self, dir:pathlib.Path):
+    def __init__(self, dir:Path):
         super().__init__(dir)
 
     def get_ubt(self) -> BuildTool:
@@ -377,24 +380,24 @@ class _Engine(_UfsNode):
 
 # {{{1 project -----------------------------------------------------------------
 class _Project(_UfsNode):
-    _path   :pathlib.Path
+    _path   :Path
     _engine :_Engine
 
-    def __init__(self, path:pathlib.Path, engine:_Engine):
+    def __init__(self, path:Path, engine:_Engine):
         super().__init__(path.parent, engine)
         self._path = path
         self._engine = engine
 
-    def get_path(self) -> pathlib.Path:
+    def get_path(self) -> Path:
         return self._path
 
-    def get_name(self) -> pathlib.Path:
+    def get_name(self) -> Path:
         return self._path.stem
 
     def get_engine(self) -> _Engine:
         return self._engine
 
-    def get_target_by_type(self, target_type:TargetType, *, create_temp:bool=True) -> Optional[_Target]:
+    def get_target_by_type(self, target_type:TargetType, *, create_temp:bool=True) -> _Target|None:
         if target_type.value <= TargetType.PROGRAM.value:
             return
 
@@ -408,7 +411,7 @@ class _Project(_UfsNode):
             target_name += type_name
 
         target = _Target(target_name, target_type, self.get_root_mount())
-        if create_temp:
+        if create_temp and Platform.get_host() != "Mac":
             self._create_temp_target(target)
 
         return target
@@ -443,29 +446,29 @@ class _Project(_UfsNode):
 
 # {{{1 branch ------------------------------------------------------------------
 class _Branch(object):
-    _dir:pathlib.Path
+    _dir:Path
 
-    def __init__(self, dir:pathlib.Path):
+    def __init__(self, dir:Path):
         self._dir = dir
 
     def get_name(self) -> str:
         return self._dir.name
 
-    def get_dir(self) -> pathlib.Path:
+    def get_dir(self) -> Path:
         return self._dir
 
-    def find_project(self, name:str) -> Optional[pathlib.Path]:
+    def find_project(self, name:str) -> Path|None:
         name = name.lower()
         for project_path in self.read_projects():
             if project_path.stem.lower() == name:
                 return project_path
 
-    def read_projects(self) -> Iterator[pathlib.Path]:
-        def read_dir(dir:pathlib.Path):
+    def read_projects(self) -> Iterator[Path]:
+        def read_dir(dir:Path):
             for sub_dir in (x for x in dir.glob("*") if x.is_dir()):
                 yield from sub_dir.glob("*.uproject")
 
-        def read_uprojectdirs(uprojdirs_path:pathlib.Path) -> Iterator[pathlib.Path]:
+        def read_uprojectdirs(uprojdirs_path:Path) -> Iterator[Path]:
             with uprojdirs_path.open() as file:
                 root_dir = uprojdirs_path.parent
                 for line in file:
@@ -483,15 +486,15 @@ class _Branch(object):
 
 # {{{1 platform-provider -------------------------------------------------------
 class _PlatformProvider(object):
-    _temp_ctx   :ForwardRef("Context")
-    _provisions :Dict[str, pathlib.Path]
+    _temp_ctx   :"Context"
+    _provisions :dict[str, Path]
 
-    def __init__(self, temp_ctx):
+    def __init__(self, temp_ctx:"Context"):
         self._temp_ctx = temp_ctx
         self._provisions = None
 
     @classmethod
-    def _get_host_platform_name(self):
+    def _get_host_platform_name(self) -> str:
         if hasattr(self, "_host_name"):
             return self._host_name
 
@@ -505,7 +508,7 @@ class _PlatformProvider(object):
 
         return self._host_name
 
-    def _get_provisions(self) -> Dict[str, pathlib.Path]:
+    def _get_provisions(self) -> dict[str, Path]:
         if self._provisions is not None:
             return self._provisions
 
@@ -513,7 +516,7 @@ class _PlatformProvider(object):
 
         # Built-in platforms
         host_name = self._get_host_platform_name()
-        for item in pathlib.Path(__file__).parent.glob(f"platforms/{host_name}/*.py"):
+        for item in Path(__file__).parent.glob(f"platforms/{host_name}/*.py"):
             if not item.name.startswith("_"):
                 provisions[item.stem] = item
 
@@ -521,7 +524,7 @@ class _PlatformProvider(object):
             return provisions
 
         # Use Engine/Platforms if we are hosted in a UE branch
-        for candidate in pathlib.Path(__file__).parent.parent.parents:
+        for candidate in Path(__file__).parent.parent.parents:
             if candidate.name == "Engine":
                 candidate = (candidate / "Platforms")
                 if candidate.is_dir():
@@ -539,7 +542,7 @@ class _PlatformProvider(object):
 
         return provisions
 
-    def get_platform(self, name:None|str=None) -> Platform|None:
+    def get_platform(self, name:str|None=None) -> Platform|None:
         host_name = self._get_host_platform_name()
         name = (name or host_name).lower()
 
@@ -561,12 +564,12 @@ class _PlatformProvider(object):
 # {{{1 context -----------------------------------------------------------------
 class Context(object):
     _type          :ContextType
-    _engine_dir    :pathlib.Path
-    _project_path  :pathlib.Path
+    _engine_dir    :Path
+    _project_path  :Path
     _platforms     :_PlatformProvider
 
     def __init__(self, primary_dir:str):
-        primary_dir:pathlib.Path = pathlib.Path(primary_dir).resolve()
+        primary_dir:Path = Path(primary_dir).resolve()
 
         # If a .uproject* was given instead the strip it off.
         if primary_dir.suffix.startswith(".uproject"):
@@ -576,7 +579,7 @@ class Context(object):
 
         # Walk up from 'primary_dir' and find something to build a context from
         if primary_dir.is_dir():
-            iter_dir:pathlib.Path = primary_dir
+            iter_dir:Path = primary_dir
             while not iter_dir.samefile(iter_dir.anchor):
                 if uproj_path := next(iter_dir.glob("*.uproject*"), None):
                     return self._initialize(uproj_path)
@@ -585,17 +588,17 @@ class Context(object):
 
         raise EnvironmentError(f"Unable to establish an Unreal context with '{primary_dir}'")
 
-    def _initialize(self, uproj_path:pathlib.Path) -> None:
+    def _initialize(self, uproj_path:Path) -> None:
         if uproj_path.suffix == ".uproject":
             return self._initialize_project(uproj_path)
         else:
             return self._initialize_noproject(uproj_path)
 
-    def _initialize_project(self, uproj_path:pathlib.Path) -> None:
+    def _initialize_project(self, uproj_path:Path) -> None:
         self._project_path = uproj_path
 
         # Find engine. First look up the tree from the project for an Engine/ dir
-        iter_dir:pathlib.Path = uproj_path.parent.parent
+        iter_dir:Path = uproj_path.parent.parent
         while not iter_dir.samefile(iter_dir.anchor):
             candidate = iter_dir / "Engine"
             if candidate.is_dir():
@@ -623,7 +626,7 @@ class Context(object):
 
         raise EnvironmentError(f"Unable to locate the engine for project '{uproj_path}'")
 
-    def _initialize_noproject(self, uprojdirs_path:pathlib.Path) -> None:
+    def _initialize_noproject(self, uprojdirs_path:Path) -> None:
         # Engine/ is expected to be a sibling of the .uprojectdirs file
         engine_dir = uprojdirs_path.parent / "Engine"
         if not engine_dir.is_dir():
@@ -643,10 +646,10 @@ class Context(object):
     def get_name(self) -> str:
         return self._engine_dir.parent.name or "UnrealEngine"
 
-    def get_primary_path(self) -> pathlib.Path:
+    def get_primary_path(self) -> Path:
         return self._project_path if self._project_path else self.get_branch().get_dir()
 
-    def get_project(self) -> Optional[_Project]:
+    def get_project(self) -> _Project|None:
         if self._type != ContextType.BRANCH:
             return _Project(self._project_path, self.get_engine())
 
@@ -656,7 +659,7 @@ class Context(object):
     def get_engine(self) -> _Engine:
         return _Engine(self._engine_dir)
 
-    def get_branch(self, *, must_exist:bool=False) -> Optional[_Branch]:
+    def get_branch(self, *, must_exist:bool=False) -> _Branch|None:
         if self._type != ContextType.FOREIGN:
             return _Branch(self._engine_dir.parent)
 
@@ -671,7 +674,7 @@ class Context(object):
     def glob(self, pattern:str, shallow:bool=False) -> Iterator[_UfsItem]:
         yield from self._get_ufs_node().glob(pattern, shallow=shallow)
 
-    def get_target_by_type(self, target_type:TargetType) -> Optional[_Target]:
+    def get_target_by_type(self, target_type:TargetType) -> _Target|None:
         return self._get_ufs_node().get_target_by_type(target_type)
 
     def get_target_by_name(self, name:str) -> _Target:
diff --git a/Engine/Extras/ushell/channels/unreal/core/pylib/unreal/platforms/linux/linux.py b/Engine/Extras/ushell/channels/unreal/core/pylib/unreal/platforms/linux/linux.py
--- a/Engine/Extras/ushell/channels/unreal/core/pylib/unreal/platforms/linux/linux.py
+++ b/Engine/Extras/ushell/channels/unreal/core/pylib/unreal/platforms/linux/linux.py
@@ -45,9 +45,10 @@ class Platform(PlatformBase):
 
     def _launch(self, exec_context, stage_dir, binary_path, args):
         if stage_dir:
-            base_dir = binary_path.replace("\\", "/").split("/")
-            base_dir = "/".join(base_dir[-4:-1])
-            base_dir = stage_dir + base_dir
+            midfix = "Engine";
+            if project := self.get_unreal_context().get_project():
+                midfix = project.get_name()
+            base_dir = stage_dir + midfix + "/Binaries/Linux"
             if not os.path.isdir(base_dir):
                 raise EnvironmentError(f"Failed to find base directory '{base_dir}'")
             args = (*args, "-basedir=" + base_dir)
diff --git a/Engine/Extras/ushell/channels/unreal/core/pylib/unreal/platforms/mac/mac.py b/Engine/Extras/ushell/channels/unreal/core/pylib/unreal/platforms/mac/mac.py
--- a/Engine/Extras/ushell/channels/unreal/core/pylib/unreal/platforms/mac/mac.py
+++ b/Engine/Extras/ushell/channels/unreal/core/pylib/unreal/platforms/mac/mac.py
@@ -27,13 +27,12 @@ class Platform(unreal.Platform):
 
     def _launch(self, exec_context, stage_dir, binary_path, args):
         if stage_dir:
-            bins_index = binary_path.find("/Binaries/Mac/")
-            if bins_index >= 0:
-                base_dir = stage_dir
-                base_dir += os.path.basename(binary_path[:bins_index])
-                base_dir += "/Binaries/Mac/"
-            else:
-                raise EnvironmentError(f"Unable to calculate base directory from '{binary_path}'")
+            midfix = "Engine";
+            if project := self.get_unreal_context().get_project():
+                midfix = project.get_name()
+            base_dir = stage_dir + midfix + "/Binaries/Mac"
+            if not os.path.isdir(base_dir):
+                raise EnvironmentError(f"Failed to find base directory '{base_dir}'")
             args = (*args, "-basedir=" + base_dir)
 
         cmd = exec_context.create_runnable(binary_path, *args)
diff --git a/Engine/Extras/ushell/channels/unreal/core/pylib/unreal/platforms/win64/android.py b/Engine/Extras/ushell/channels/unreal/core/pylib/unreal/platforms/win64/android.py
--- a/Engine/Extras/ushell/channels/unreal/core/pylib/unreal/platforms/win64/android.py
+++ b/Engine/Extras/ushell/channels/unreal/core/pylib/unreal/platforms/win64/android.py
@@ -4,10 +4,13 @@ import os
 import re
 import shutil
 import unreal
+import subprocess as sp
+from pathlib import Path
 
 #-------------------------------------------------------------------------------
 class Platform(unreal.Platform):
     name = "Android"
+    config_section = "/Script/AndroidRuntimeSettings.AndroidRuntimeSettings"
 
     def _read_env(self):
         version = self.get_version()
@@ -76,12 +79,13 @@ class Platform(unreal.Platform):
             return
 
         for prefix in ("", unreal.Platform.get_sdks_dir()):
-            if os.path.isdir(prefix + out):
-                return prefix + out
+            candidate = Path(prefix) / out
+            if candidate.is_dir():
+                return candidate
 
     def _get_adb(self):
         home_dir = self._get_android_home()
-        return home_dir + "/platform-tools/adb" if home_dir else "adb"
+        return (home_dir / "platform-tools/adb") if home_dir else "adb"
 
     def _get_aapt(self):
         home_dir = self._get_android_home()
@@ -89,15 +93,60 @@ class Platform(unreal.Platform):
             return "aapt"
 
         try:
-            build_tools_dir = home_dir + "/build-tools/"
-            latest = ""
-            for item in os.scandir(build_tools_dir):
-                latest = max(latest, item.name)
-            if latest and os.path.isdir(build_tools_dir):
-                return build_tools_dir + latest + "/aapt"
-        except FileNotFoundError:
+            build_tools_dir = home_dir / "build-tools"
+            latest = max(x.name for x in build_tools_dir.glob("*"))
+            return build_tools_dir / latest / "aapt"
+        except ValueError:
             return
 
     def _get_cook_form(self, target):
         if target == "game":   return "Android_ASTC"
         if target == "client": return "Android_ASTCClient"
+
+    def _get_package_name(self):
+        config = self.get_unreal_context().get_config()
+        if ret := str(config.get("Engine", Platform.config_section, "PackageName")):
+            return ret
+        raise ValueError(f"Failed querying '{Platform.config_section}/PackageName'")
+
+    def _launch(self, exec_context, stage_dir, binary_path, args):
+        print("!!\n!! NOTE: .run does not currently use locally-built executables\n!!\n")
+
+        package_name = self._get_package_name()
+        intent = package_name + "/com.epicgames.unreal.SplashActivity"
+        print("Package:", package_name)
+        print("Intent: ", intent)
+        print()
+
+        args = "'" + " ".join(x.replace("\"", "\\\"") for x in args) + "'"
+        adb = self._get_adb()
+        cmd = (str(adb), "-d", "shell", "am", "start", "--es", "cmdline", args, intent)
+        print(*cmd)
+
+        # subprocess.run() messes with quotes on Windows so we'll use os.system
+        self._kill()
+        os.system(" ".join(cmd))
+
+        return True
+        """ -- not yet implemented properly --
+        # get device
+        device = None
+        for line in exec_context.create_runnable(adb, "-d", "devices"):
+            if line.endswith("device") and not device:
+                device = line.split()[0]
+        print(device)
+
+        # get pid
+        pid = -1
+        for line in exec_context.create_runnable(adb, "-d", "shell", "run-as", package_name, "ps"):
+            print(line)
+
+        return (device, pid)
+        """
+
+    def _kill(self, target=""):
+        package_name = self._get_package_name()
+        adb = self._get_adb()
+        cmd = (adb, "-d", "shell", "am", "force-stop", str(package_name))
+        print(*cmd)
+        sp.run(cmd)
diff --git a/Engine/Extras/ushell/channels/unreal/core/pylib/unreal/platforms/win64/win64.py b/Engine/Extras/ushell/channels/unreal/core/pylib/unreal/platforms/win64/win64.py
--- a/Engine/Extras/ushell/channels/unreal/core/pylib/unreal/platforms/win64/win64.py
+++ b/Engine/Extras/ushell/channels/unreal/core/pylib/unreal/platforms/win64/win64.py
@@ -119,9 +119,10 @@ class Platform(_PlatformBase):
 
     def _launch(self, exec_context, stage_dir, binary_path, args):
         if stage_dir:
-            base_dir = binary_path.replace("\\", "/").split("/")
-            base_dir = "/".join(base_dir[-4:-1])
-            base_dir = stage_dir + base_dir
+            midfix = "Engine";
+            if project := self.get_unreal_context().get_project():
+                midfix = project.get_name()
+            base_dir = stage_dir + midfix + "/Binaries/Win64"
             if not os.path.isdir(base_dir):
                 raise EnvironmentError(f"Failed to find base directory '{base_dir}'")
             args = (*args, "-basedir=" + base_dir)
diff --git a/Engine/Extras/ushell/channels/unreal/core/pylib/unrealcmd.py b/Engine/Extras/ushell/channels/unreal/core/pylib/unrealcmd.py
--- a/Engine/Extras/ushell/channels/unreal/core/pylib/unrealcmd.py
+++ b/Engine/Extras/ushell/channels/unreal/core/pylib/unrealcmd.py
@@ -91,7 +91,8 @@ class MultiPlatformCmd(Cmd):
                 print(item.key, "=", dir)
             return context
 
-        self.print_info("Establishing platforms")
+        if self.is_interactive():
+            self.print_info("Establishing platforms")
 
         found = []
         failed = []
@@ -104,12 +105,13 @@ class MultiPlatformCmd(Cmd):
             except EnvironmentError:
                 failed.append(platform_name)
 
-        if found:
-            print("  found:", flow.cmd.text.green(" ".join(found)))
-        if failed:
-            out = "missing: "
-            out += flow.cmd.text.light_red(" ".join(failed))
-            out += flow.cmd.text.grey(" ('.info' for details)")
-            print(out)
+        if self.is_interactive():
+            if found:
+                print("  found:", flow.cmd.text.green(" ".join(found)))
+            if failed:
+                out = "missing: "
+                out += flow.cmd.text.light_red(" ".join(failed))
+                out += flow.cmd.text.grey(" ('.info' for details)")
+                print(out)
 
         return context
diff --git a/Engine/Extras/ushell/channels/unreal/perforce/cmds/bisect.py b/Engine/Extras/ushell/channels/unreal/perforce/cmds/bisect.py
--- a/Engine/Extras/ushell/channels/unreal/perforce/cmds/bisect.py
+++ b/Engine/Extras/ushell/channels/unreal/perforce/cmds/bisect.py
@@ -1,6 +1,7 @@
 # Copyright Epic Games, Inc. All Rights Reserved.
 
 import os
+import sys
 import p4utils
 import unrealcmd
 from peafour import P4
@@ -100,7 +101,10 @@ class Bisect(unrealcmd.Cmd):
         if self.args.dryrun:
             return _user_choice()
 
-        ret = sp.run(self._script).returncode
+        cmd = self._script
+        if self._script[0].endswith(".py"):
+            cmd = (sys.executable, *cmd)
+        ret = sp.run(cmd).returncode
         if ret == 80: return "b"
         if ret ==  0: return "g"
         if ret == 90: return "f"
@@ -272,7 +276,6 @@ def _test_bisectomatron():
             print(" ", line.decode(), sep="")
             print()
             print(b.get_good(), b.get_bad())
-            breakpoint()
 
     for i in range(50000):
         print(i, end="\r")
diff --git a/Engine/Extras/ushell/channels/unreal/perforce/cmds/cherrypick.py b/Engine/Extras/ushell/channels/unreal/perforce/cmds/cherrypick.py
--- a/Engine/Extras/ushell/channels/unreal/perforce/cmds/cherrypick.py
+++ b/Engine/Extras/ushell/channels/unreal/perforce/cmds/cherrypick.py
@@ -35,6 +35,7 @@ class Cherrypick(flow.cmd.Cmd):
     saferesolve   = flow.cmd.Opt(False, "Resolve safely and not automatically")
     dryrun        = flow.cmd.Opt(False, "Only pretend to do the cherrypick")
     force         = flow.cmd.Opt(False, "Force the operation through")
+    novalidate    = flow.cmd.Opt(False, "Don't run the validation step")
     alwayseddy    = flow.cmd.Opt(False, "Always edigrate the result regardless of relation")
     noeddy        = flow.cmd.Opt(False, "Skip the edigrate step")
     sync          = flow.cmd.Opt(False, "Syncs target files to head before resolving")
@@ -163,7 +164,7 @@ class Cherrypick(flow.cmd.Cmd):
         print("Affected files:", len(dest_paths))
         print("Checking for open files")
 
-        validate = not self.args.force
+        validate = not (self.args.novalidate or self.args.force)
         if len(dest_paths) > 500:
             validate = False
             self.print_warning("Skipping validation. Source changelist is too big")
diff --git a/Engine/Extras/ushell/channels/unreal/perforce/cmds/sync.py b/Engine/Extras/ushell/channels/unreal/perforce/cmds/sync.py
--- a/Engine/Extras/ushell/channels/unreal/perforce/cmds/sync.py
+++ b/Engine/Extras/ushell/channels/unreal/perforce/cmds/sync.py
@@ -296,26 +296,28 @@ class Sync(_SyncBase):
 
         except KeyboardInterrupt:
             print()
-            self.print_warning(f"Sync interrupted! Writing build.version to CL {version_cl}")
+            if not self.args.dryrun:
+                self.print_warning(f"Sync interrupted! Writing build.version to CL {version_cl}")
             return False
         finally:
-            # Record the synced changelist in Build.version
-            with open(build_ver_path, "r") as x:
-                lines = list(x.readlines())
-
-            import stat
-            build_ver_mode = os.stat(build_ver_path).st_mode
-            os.chmod(build_ver_path, build_ver_mode|stat.S_IWRITE)
-
-            with open(build_ver_path, "w") as x:
-                for line in lines:
-                    if r'"Changelist"' in line:
-                        line = line.split(":", 2)
-                        line = line[0] + f": {version_cl},\n"
-                    elif r'"BranchName"' in line:
-                        line = "\t\"BranchName\": \"X\"\n"
-                        line = line.replace("X", self._branch_root[:-1].replace("/", "+"))
-                    x.write(line)
+            if not self.args.dryrun:
+                # Record the synced changelist in Build.version
+                with open(build_ver_path, "r") as x:
+                    lines = list(x.readlines())
+
+                import stat
+                build_ver_mode = os.stat(build_ver_path).st_mode
+                os.chmod(build_ver_path, build_ver_mode|stat.S_IWRITE)
+
+                with open(build_ver_path, "w") as x:
+                    for line in lines:
+                        if r'"Changelist"' in line:
+                            line = line.split(":", 2)
+                            line = line[0] + f": {version_cl},\n"
+                        elif r'"BranchName"' in line:
+                            line = "\t\"BranchName\": \"X\"\n"
+                            line = line.replace("X", self._branch_root[:-1].replace("/", "+"))
+                        x.write(line)
 
     def _apply_p4sync_txt(self, syncer):
         def impl(path):
diff --git a/Engine/Extras/ushell/channels/unreal/perforce/cmds/workspace.py b/Engine/Extras/ushell/channels/unreal/perforce/cmds/workspace.py
--- a/Engine/Extras/ushell/channels/unreal/perforce/cmds/workspace.py
+++ b/Engine/Extras/ushell/channels/unreal/perforce/cmds/workspace.py
@@ -57,8 +57,11 @@ def _populate_ue_branch_engine(depot_path, dry_run):
         print(*minimal_sync_specs, sep="\n")
         return
 
-    for item in P4.sync(*minimal_sync_specs):
+    count = 0
+    for item in P4.sync(*minimal_sync_specs).read(on_error=False):
+        count += 1
         print(item.depotFile)
+    print("synced:", count)
 
 #-------------------------------------------------------------------------------
 def _populate_ue_branch_project(depot_path, dry_run):
@@ -90,7 +93,7 @@ def _populate_ue_branch_project(depot_path, dry_run):
 
 #-------------------------------------------------------------------------------
 class Workspace(flow.cmd.Cmd):
-    """ Creates a new workspace """
+    """ Creates a new workspace and sync the bare minimum. """
     localdir  = flow.cmd.Arg(str, "Directory to create the new workspace in")
     depotpath = flow.cmd.Arg("", "Depot path to map the workspace to")
     name      = flow.cmd.Opt("", "The name of the workspace to use")
@@ -178,3 +181,25 @@ class Workspace(flow.cmd.Cmd):
             _populate_ue_branch(depot_path, self.args.dryrun)
         else:
             print("New clientspec does not appear to be an Unreal Engine branch")
+
+
+
+#-------------------------------------------------------------------------------
+class MinSync(flow.cmd.Cmd):
+    """ Syncs just enough to make a branch functional with ushell """
+    dryrun = flow.cmd.Opt(False, "Try but don't try too hard")
+
+    def main(self):
+        p4utils.login()
+
+        self.print_info("Perforce environment")
+        info = P4.info().run()
+        depot_path = f"//{info.clientName}"
+
+        print("Cwd:", os.getcwd())
+        print("Client:", info.clientName)
+        print("Root:", depot_path)
+
+        self.print_info("Syncing")
+        depot_path = PurePosixPath(depot_path)
+        _populate_ue_branch(depot_path, self.args.dryrun)
diff --git a/Engine/Extras/ushell/channels/unreal/perforce/describe.flow.py b/Engine/Extras/ushell/channels/unreal/perforce/describe.flow.py
--- a/Engine/Extras/ushell/channels/unreal/perforce/describe.flow.py
+++ b/Engine/Extras/ushell/channels/unreal/perforce/describe.flow.py
@@ -22,6 +22,11 @@ client = flow.describe.Command()
 client.source("cmds/workspace.py", "Workspace")
 client.invoke("p4", "workspace")
 
+#-------------------------------------------------------------------------------
+min_sync = flow.describe.Command()
+min_sync.source("cmds/workspace.py", "MinSync")
+min_sync.invoke("p4", "sync", "mini")
+
 #-------------------------------------------------------------------------------
 gui = flow.describe.Command()
 gui.source("cmds/gui.py", "Gui")
