ushell: add foreign project support for Linux and Mac

diff --git a/Engine/Extras/ushell/channels/unreal/core/pylib/unreal/_context.py b/Engine/Extras/ushell/channels/unreal/core/pylib/unreal/_context.py
--- a/Engine/Extras/ushell/channels/unreal/core/pylib/unreal/_context.py
+++ b/Engine/Extras/ushell/channels/unreal/core/pylib/unreal/_context.py
@@ -35,7 +35,28 @@ def _resolve_association_nt(association:str) -> Path|None:
         return Path(value[0])
 
 def _resolve_association_posix(association:str) -> Path|None:
-    raise NotImplementedError("Foreign projects are not yet supported")
+    import configparser, sys
+
+    # GUID values might be stored with or without curly braces, depending on platform and engine version
+    candidates = [association]
+    if association.startswith("{") and association.endswith("}"):
+        candidates.append(association[1:-1])
+    else:
+        candidates.append(f"{{{association}}}")
+
+    # Official engine releases are prefixed with "UE_" in Install.ini under Linux, but not in .uproject files
+    candidates.append(f"UE_{association}")
+
+    config_dir = "~/.config" if sys.platform == "linux" else "~/Library/Application Support"
+    ini_file = Path(config_dir).expanduser() / "Epic" / "UnrealEngine" / "Install.ini"
+
+    config = configparser.ConfigParser()
+    config.read([ini_file])
+
+    for candidate in candidates:
+        value = config.get("Installations", candidate, fallback=None)
+        if value is not None:
+            return Path(value)
 
 def _resolve_association(association:str) -> Path|None:
     if os.name == "nt":
