ushell: don't attempt to build ShaderCompileWorker and UnrealPak targets for Installed Builds

diff --git a/Engine/Extras/ushell/channels/unreal/core/cmds/build.py b/Engine/Extras/ushell/channels/unreal/core/cmds/build.py
--- a/Engine/Extras/ushell/channels/unreal/core/cmds/build.py
+++ b/Engine/Extras/ushell/channels/unreal/core/cmds/build.py
@@ -439,12 +439,17 @@ class Editor(_BuildCmd, unrealcmd.MultiPlatformCmd):
         editor_only |= bool(self.args.analyze)
         editor_only |= self.args.variant != "development"
         if not editor_only:
+
+            targets_to_add = []
             if not self.args.noscw:
-                targets = (*targets, ue_context.get_target_by_name("ShaderCompileWorker"))
+                targets_to_add.append("ShaderCompileWorker")
             if not self.args.nopak:
-                targets = (*targets, ue_context.get_target_by_name("UnrealPak"))
+                targets_to_add.append("UnrealPak")
             if not self.args.nointworker:
-                prog_target = ue_context.get_target_by_name("InterchangeWorker")
+                targets_to_add.append("InterchangeWorker")
+
+            for name in targets_to_add:
+                prog_target = ue_context.get_target_by_name(name)
                 if prog_target.get_type() == unreal.TargetType.PROGRAM:
                     targets = (*targets, prog_target)
 
