ushell: add child shell workflow for POSIX platforms

diff --git a/Engine/Extras/ushell/channels/flow/posix/boot.sh b/Engine/Extras/ushell/channels/flow/posix/boot.sh
--- a/Engine/Extras/ushell/channels/flow/posix/boot.sh
+++ b/Engine/Extras/ushell/channels/flow/posix/boot.sh
@@ -18,7 +18,7 @@ if [[ -n "$FLOW_CHANNELS_DIR" ]]; then
     channels_dir+=("$FLOW_CHANNELS_DIR")
 fi
 
-"$working_dir/python/current/bin/python" -Esu "$self_dir/../core/system/boot.py" "$working_dir" "${channels_dir[@]}" -- $@
+"$working_dir/python/current/bin/python" -Esu "$self_dir/../core/system/boot.py" "$working_dir" "${channels_dir[@]}" -- "$@"
 if [ ! $? ]; then
     echo boot.py failed
     exit 1
diff --git a/Engine/Extras/ushell/ushell.sh b/Engine/Extras/ushell/ushell.sh
--- a/Engine/Extras/ushell/ushell.sh
+++ b/Engine/Extras/ushell/ushell.sh
@@ -1,12 +1,18 @@
 #!/bin/bash
 
-# Check that this script has been "sourced" into a host shell
+# Check whether this script has been "sourced" into a host shell
 if [ $0 = "$BASH_SOURCE" ]; then
-    echo You must source ushell into your host shell\;
-    echo
-    echo "    source $0"
-    echo
-    exit 1
+    
+    # The script was not sourced, so source it in a child shell and propagate any arguments
+    args=""
+    for arg in "$@"
+    do
+        escaped=$(printf '%q' "$arg")
+        args="$args $escaped"
+    done
+    bash --init-file <(echo "source \"$HOME/.bashrc\"; source \"$BASH_SOURCE\" $args")
+    exit $?
+    
 fi
 
 # Determine the host shell
@@ -33,7 +39,7 @@ else
     cookie=/tmp/ushell_$$_shell_cookie
 fi
 
-$(dirname ${BASH_SOURCE:-$0})/channels/flow/$channel/boot.sh --bootarg=$host_shell,$cookie $@
+$(dirname ${BASH_SOURCE:-$0})/channels/flow/$channel/boot.sh --bootarg=$host_shell,$cookie "$@"
 
 if [ -f $cookie ]; then
     chmod u+x $cookie
