Updated ushell [32355311]

diff --git a/Engine/Extras/ushell/channels/flow/core/cmds/debug.py b/Engine/Extras/ushell/channels/flow/core/cmds/debug.py
--- a/Engine/Extras/ushell/channels/flow/core/cmds/debug.py
+++ b/Engine/Extras/ushell/channels/flow/core/cmds/debug.py
@@ -11,8 +11,15 @@ from urllib.request import urlopen, URLError
 
 #-------------------------------------------------------------------------------
 def _http_get(url, on_data):
+    # Try and get the certifi CA file
+    try:
+        import certifi
+        cafile_path = certifi.where()
+    except ImportError:
+        cafile_path = None
+
     # Fire up an http client to get the url.
-    client = urlopen(url)
+    client = urlopen(url, cafile=cafile_path)
     if (client.status // 100) != 2:
         assert False, f"Error creating HTTP client for {url}"
 
@@ -67,6 +74,8 @@ class Invalidate(flow.cmd.Cmd):
                 print("Removing:", manifest_path.relative_to(working_dir))
                 manifest_path.unlink(missing_ok=True)
 
+        (working_dir / "manifest").unlink(missing_ok=True)
+
 #-------------------------------------------------------------------------------
 class Wipe(flow.cmd.Cmd):
     """ Wipes everything. """
diff --git a/Engine/Extras/ushell/channels/flow/core/system/flow/describe.py b/Engine/Extras/ushell/channels/flow/core/system/flow/describe.py
--- a/Engine/Extras/ushell/channels/flow/core/system/flow/describe.py
+++ b/Engine/Extras/ushell/channels/flow/core/system/flow/describe.py
@@ -48,33 +48,15 @@ class Command(object):
         self._py_path = script_path
         self._py_class = cmd_class
 
-#-------------------------------------------------------------------------------
-class Extension(object):
-    def __init__(self):
-        self._name = None
-        self._mount = None
-        self._py_path = None
-        self._py_class = None
-
-    def mount(self, extendable, name):
-        self._extendable = extendable
-        self._name = name
-
-    def source(self, script_path, ext_class):
-        self._py_path = script_path
-        self._py_class = ext_class
-
 #-------------------------------------------------------------------------------
 class Channel(object):
     def __init__(self):
         self._parent = ""
         self._pips = []
-        self._extendables = []
         self._version = None
 
     def parent(self, parent):   self._parent = parent
     def version(self, version): self._version = version
-    def extendable(self, name): self._extendables.append(name)
 
     def _pip_deprecated_(self, name):
         # Please do not use! Pips come with security and licensing headaches so
diff --git a/Engine/Extras/ushell/channels/flow/core/system/lib/bootstrap.py b/Engine/Extras/ushell/channels/flow/core/system/lib/bootstrap.py
--- a/Engine/Extras/ushell/channels/flow/core/system/lib/bootstrap.py
+++ b/Engine/Extras/ushell/channels/flow/core/system/lib/bootstrap.py
@@ -48,8 +48,15 @@ def _http_get(url, dest_dir, progress_cb=None):
     import re
     from urllib.request import urlopen
 
+    # Try and get the certifi CA file
+    try:
+        import certifi
+        cafile_path = certifi.where()
+    except ImportError:
+        cafile_path = None
+
     # Fire up an http client to get the url.
-    client = urlopen(url)
+    client = urlopen(url, cafile=cafile_path)
     if (client.status // 100) != 2:
         assert False, f"Error creating HTTP client for {url}"
 
@@ -85,17 +92,6 @@ def _http_get(url, dest_dir, progress_cb=None):
 
 
 
-#-------------------------------------------------------------------------------
-def _extract_zip(payload, dest_dir):
-    # Use tar on POSIX platforms as Python's zipfile doesn't seem to apply attrs
-    if os.name != "nt":
-        return _extract_tar(payload, dest_dir)
-
-    import zipfile
-    zip_file = zipfile.ZipFile(payload, "r")
-    zip_file.extractall(dest_dir)
-    zip_file.close()
-
 #-------------------------------------------------------------------------------
 def _extract_tar(payload, dest_dir):
     cmd = ("tar", "-x", "-C", dest_dir, "-f", payload)
@@ -107,7 +103,7 @@ def _tool_extract_method(self, src_path, dest_dir):
     extractors = {
         ".tar.gz"   : _extract_tar,
         ".tgz"      : _extract_tar,
-        ".zip"      : _extract_zip,
+        ".zip"      : _extract_tar, # Windows comes with tar as standard now
     }
 
     was_archive = False
@@ -151,9 +147,9 @@ def _acquire_tool(name, tool, target_dir, progress_cb):
         _log.write("Bundle", bundle_name)
         fetch_dir = temp_dir + bundle_name + "/"
         os.mkdir(fetch_dir)
-        for payload in bundle._payloads:
-            _log.write("Payload", *payload)
-            payload, ftype = payload
+        for payload, ftype in bundle._payloads:
+            payload = payload.replace("$VERSION", tool._version)
+            _log.write("Payload", payload, ftype)
             payload_dest = _http_get(payload, fetch_dir, inner_progress_cb)
             payload_i += 1
             if ftype:
@@ -214,7 +210,8 @@ def _acquire_tool(name, tool, target_dir, progress_cb):
 def _manifest_tool(name, tool):
     bundles = {}
     for bundle_name, bundle in tool._bundles.items():
-        bundles[bundle_name] = tuple(x[0] for x in bundle._payloads)
+        urls = (x[0] for x in bundle._payloads)
+        bundles[bundle_name] = tuple(x.replace("$VERSION", tool._version) for x in urls)
 
     return {
         "sha1"      : 0,
@@ -295,14 +292,6 @@ def _validate_command(channel_src_dir, name, command):
     if not os.path.isfile(channel_src_dir + command._py_path):
         assert False, f"Command '{name}' missing source '{command._py_path}'"
 
-#-------------------------------------------------------------------------------
-def _validate_extension(channel_src_dir, name, extension):
-    assert extension._name, f"Extension '{name}' invalid/missing call to mount()"
-    assert extension._extendable, f"Extension '{name}' missing call to mount()"
-    assert extension._py_path, f"Extension '{name}' missing call to source()"
-    if not os.path.isfile(channel_src_dir + extension._py_path):
-        assert False, f"Extension '{name}' missing source '{extension._py_path}'"
-
 #-------------------------------------------------------------------------------
 def _select_update_type():
     version_path = os.path.abspath(__file__ + "/../../version")
@@ -370,21 +359,13 @@ class _UpdateImpl(_Action):
             command._path = command._path or (name)
             command._path = (command._prefix, *command._path)
 
-        # Extensions
-        extensions = [x for x in read_items_by_type(flow.describe.Extension)]
-
-        for name, extension in extensions:
-            _validate_extension(self._input_path, name, extension)
-
         # Build manifest
         manifest = {}
         manifest["name"] = self._name
         manifest["path"] = self._input_path
         manifest["parent"] = channel._parent or "flow.core"
-        manifest["extendables"] = tuple(channel._extendables)
         manifest["tools"] = tools_manifest
         manifest["commands"] = tuple(v.__dict__ for k,v in commands)
-        manifest["extensions"] = tuple(v.__dict__ for k,v in extensions)
 
         # Success. Write the action's output
         with open(out_dir + "manifest", "wb") as x:
@@ -425,29 +406,6 @@ def _plant_cmd_tree(channels):
 
     return cmd_tree
 
-#-------------------------------------------------------------------------------
-def _mount_extensions(channel):
-    py_dir = channel["name"] + "/"
-    for extension in channel["extensions"]:
-        extendable = extension["_extendable"]
-        host = channel
-        while True:
-            if extendable in host["extendables"]:
-                ext_desc = (
-                    channel["index"],
-                    extension["_py_path"],
-                    extension["_py_class"],
-                )
-                out = host["extendables"][extendable]
-                out = out.setdefault(extension["_name"], [])
-                out.append(ext_desc)
-                break
-
-            host = host["parent"]
-            if not host:
-                extension_name = extension["_name"]
-                assert False, f"Unable to find extendable '{extendable}' for '{py_dir}{extension_name}'"
-
 #-------------------------------------------------------------------------------
 def _carve_channels(channels):
     out = [None] * len(channels)
@@ -459,7 +417,6 @@ def _carve_channels(channels):
             "name" : channel["name"],
             "path" : channel["path"],
             "parent" : parent_index,
-            "extendables" : channel["extendables"],
             "tools" : tools,
         }
 
@@ -537,11 +494,6 @@ class _Finalise(_Action):
         manifests = topo_manifests
         _log.write("Channels;", *(x["name"] for x in manifests.values()))
 
-        # Prepare extendables and mount extensions
-        for manifest in manifests.values():
-            manifest["extendables"] = {x:{} for x in manifest["extendables"]}
-            _mount_extensions(manifest)
-
         # Collect various paths for each channel
         pylib_dirs = []
         for manifest in manifests.values():
diff --git a/Engine/Extras/ushell/channels/flow/core/system/version b/Engine/Extras/ushell/channels/flow/core/system/version
--- a/Engine/Extras/ushell/channels/flow/core/system/version
+++ b/Engine/Extras/ushell/channels/flow/core/system/version
@@ -1 +1 @@
-20
+21
diff --git a/Engine/Extras/ushell/channels/flow/nt/boot.bat b/Engine/Extras/ushell/channels/flow/nt/boot.bat
--- a/Engine/Extras/ushell/channels/flow/nt/boot.bat
+++ b/Engine/Extras/ushell/channels/flow/nt/boot.bat
@@ -34,9 +34,14 @@ set _cookie=%_cookie:~1%
 set _cookie=%temp%\ushell\cmd_boot_%_cookie%
 %_working%\python\current\flow_python.exe -Xutf8 -Esu "%~dp0..\core\system\boot.py" %_working% %_channels% -- "--bootarg=cmd,%_cookie%" %*
 
-:: was the python binary missing? "3" comes from cmd.exe - lets hope that it
-:: doesn't collide with something boot.py might return.
-if errorlevel 3 (
+:: If --help was given then we return the user to where they came from. The "127"
+:: comes from _flick
+if %errorlevel%==127 (
+    goto:eof
+)
+
+:: Was the Python binary missing? "9009" comes from cmd.exe
+if %errorlevel%==9009 (
     echo,
     echo,
     echo ERROR: Missing flow_python.exe from '%_working%\python\current'
diff --git a/Engine/Extras/ushell/channels/flow/nt/describe.flow.py b/Engine/Extras/ushell/channels/flow/nt/describe.flow.py
--- a/Engine/Extras/ushell/channels/flow/nt/describe.flow.py
+++ b/Engine/Extras/ushell/channels/flow/nt/describe.flow.py
@@ -6,7 +6,7 @@ import flow.describe
 #-------------------------------------------------------------------------------
 clink = flow.describe.Tool()
 clink.version("1.0.0a4")
-clink.payload("[REDACTED_PROFILE_URL]/clink/releases/download/1.0.0a4/clink-1.0.0a4.zip")
+clink.payload("[REDACTED_PROFILE_URL]/clink/releases/download/$VERSION/clink-$VERSION.zip")
 clink.sha1("6a642bc48e21a3093698067f0fe09eff3ccdfaa5")
 clink.platform("win32")
 clink.bin("clink_x64.exe")
@@ -14,7 +14,7 @@ clink.bin("clink_x64.exe")
 #-------------------------------------------------------------------------------
 fd = flow.describe.Tool()
 fd.version("8.7.1")
-fd.payload("[REDACTED_PROFILE_URL]/fd/releases/download/v8.7.1/fd-v8.7.1-x86_64-pc-windows-msvc.zip")
+fd.payload("[REDACTED_PROFILE_URL]/fd/releases/download/v$VERSION/fd-v$VERSION-x86_64-pc-windows-msvc.zip")
 fd.sha1("7e8f5d0a9d1fed75b12483809f9619c04b39b60c")
 fd.platform("win32")
 fd.bin("fd.exe")
@@ -23,7 +23,7 @@ fd.source("[REDACTED_PROFILE_URL]/fd/releases/latest", "fd-v([0-9.]+)-x86_64
 #-------------------------------------------------------------------------------
 fzf_win32 = flow.describe.Tool()
 fzf_win32.version("0.44.1")
-fzf_win32.payload("[REDACTED_PROFILE_URL]/fzf/releases/download/0.44.1/fzf-0.44.1-windows_amd64.zip")
+fzf_win32.payload("[REDACTED_PROFILE_URL]/fzf/releases/download/$VERSION/fzf-$VERSION-windows_amd64.zip")
 fzf_win32.sha1("463f064d0403714ab5bf68ebf9fdea4c0159f6c3")
 fzf_win32.platform("win32")
 fzf_win32.bin("fzf.exe")
diff --git a/Engine/Extras/ushell/channels/flow/nt/shell_cmd.py b/Engine/Extras/ushell/channels/flow/nt/shell_cmd.py
--- a/Engine/Extras/ushell/channels/flow/nt/shell_cmd.py
+++ b/Engine/Extras/ushell/channels/flow/nt/shell_cmd.py
@@ -6,7 +6,7 @@
 from pathlib import Path
 
 #-------------------------------------------------------------------------------
-class Cmd(object):
+class _Shell(object):
     def _add_clink_lua(self):
         channel = self.get_channel()
         system = channel.get_system()
@@ -30,39 +30,10 @@
         cmd = cmd.replace("\\", "/")
         lua.write(f"local py_server_cmd = '{cmd}'")
 
-        lua.write(_clink_lua)
+        clink_lua = _get_clink_lua()
+        lua.write(clink_lua)
         lua.close()
 
-    def run(self, env):
-        # Give the user a convenient way to add customisations in to the session
-        self.print_info("Customisation")
-        home_dir = Path(self.get_home_dir())
-        tags = (
-            flow.cmd.text.grey("none"),
-            flow.cmd.text.white("okay"),
-        )
-
-        # These are really a no-op. It is here to hint to the user how they can
-        # add their own channels
-        candidate = home_dir / "channels"
-        found = candidate.is_dir()
-        print(tags[int(found)], candidate / "*")
-
-        # Indicate to the user that there's a way to hook the boot process too
-        candidate = home_dir / "hooks/boot.bat"
-        found = candidate.is_file()
-        print(tags[int(found)], candidate)
-
-        # This is run just before the prompt is shown
-        candidate = home_dir / "hooks/startup.bat"
-        found = candidate.is_file()
-        print(tags[int(found)], candidate)
-        if found:
-            self._user_script = candidate
-
-        print()
-        super().run(env)
-
     def register_shells(self, registrar):
         registrar.add("cmd", lambda x: self)
         return super().register_shells(registrar)
@@ -106,7 +77,44 @@
 
 
 #-------------------------------------------------------------------------------
-_clink_lua = r"""
+class Cmd(_Shell):
+    def run(self, env):
+        if os.name != "nt":
+            return super().run(env)
+
+        # Give the user a convenient way to add customisations in to the session
+        self.print_info("Customisation")
+        home_dir = Path(self.get_home_dir())
+        tags = (
+            flow.cmd.text.grey("none"),
+            flow.cmd.text.white("okay"),
+        )
+
+        # These are really a no-op. It is here to hint to the user how they can
+        # add their own channels
+        candidate = home_dir / "channels"
+        found = candidate.is_dir()
+        print(tags[int(found)], candidate / "*")
+
+        # Indicate to the user that there's a way to hook the boot process too
+        candidate = home_dir / "hooks/boot.bat"
+        found = candidate.is_file()
+        print(tags[int(found)], candidate)
+
+        # This is run just before the prompt is shown
+        candidate = home_dir / "hooks/startup.bat"
+        found = candidate.is_file()
+        print(tags[int(found)], candidate)
+        if found:
+            self._user_script = candidate
+
+        print()
+        return super().run(env)
+
+
+#-------------------------------------------------------------------------------
+def _get_clink_lua():
+    return r"""
 local py_server = io.popen2(py_server_cmd)
 
 for _, tree_root in ipairs(commands) do
diff --git a/Engine/Extras/ushell/channels/flow/nt/shell_pwsh.py b/Engine/Extras/ushell/channels/flow/nt/shell_pwsh.py
--- a/Engine/Extras/ushell/channels/flow/nt/shell_pwsh.py
+++ b/Engine/Extras/ushell/channels/flow/nt/shell_pwsh.py
@@ -33,33 +33,41 @@ class _Shell(object):
         start_dir = os.getcwd()
 
         with open(cookie, "wt") as out:
+            header = _get_header_script()
+            autocompleter = _get_autocompleter_script()
+            function_template = _get_function_template_script()
+            cleanup = _get_cleanup_script()
+
             out.write(f"Set-Location \"{start_dir}\"\n")
             # Environment must be written out first so autocomplete daemon gets the right session id
             for key, value in env.read_changes():
                 value = value or ""
                 out.write(f'${{env:{key}}}="{value}"\n')
-            out.write(_header.format(shims_path=shims_path,working_dir=working_dir))
-            out.write(_autocompleter.format(python=sys.executable, run_py_path=run_py_path, working_dir=working_dir))
+            out.write(header.format(shims_path=shims_path,working_dir=working_dir))
+            out.write(autocompleter.format(python=sys.executable, run_py_path=run_py_path, working_dir=working_dir))
             for name,_ in tree_root.read_children():
                 if name.startswith("$"): continue
-                out.write(_function_template.format(name=name))
-            out.write(_cleanup)
+                out.write(function_template.format(name=name))
+            out.write(cleanup)
 
-_header = r"""
+def _get_header_script():
+    return r"""
 $ShimsPath = Resolve-Path "{shims_path}"
 $env:Path += ";$ShimsPath"
 """
 
 # Braces in this string must be doubled to escape formatting
 #TODO: give these functions long names, then optionally add aliases?
-_function_template = r"""
+def _get_function_template_script():
+    return r"""
 Register-ArgumentCompleter -CommandName "{name}.exe" -Native -ScriptBlock $SharedCompleter
 Register-ArgumentCompleter -CommandName "{name}" -Native -ScriptBlock $SharedCompleter
 # Explicitly export this function so that others are not exported automatically
 Export-ModuleMember -Function {name}
 """
 # If we give commands long names, we need to separate completers that provide the name ushell $complete expects
-_autocompleter = r"""
+def _get_autocompleter_script():
+    return r"""
 function Get-Daemon {{
     if( $script:CompleteDaemonCache -ne $null -and !$script:CompleteDaemonCache.HasExited) {{
         return $script:CompleteDaemonCache
@@ -201,7 +209,8 @@ Export-ModuleMember -Function "Update-UShellEnvVars"
 """
 
 # Kill the autocomplete daemon so we don't leak a process
-_cleanup = r"""
+def _get_cleanup_script():
+    return r"""
 $ExecutionContext.SessionState.Module.OnRemove += {
     if( $script:CompleteDaemonCache -ne $null ) {
         Write-Debug "Stopping complete daemon"
diff --git a/Engine/Extras/ushell/channels/flow/posix/provision.sh b/Engine/Extras/ushell/channels/flow/posix/provision.sh
--- a/Engine/Extras/ushell/channels/flow/posix/provision.sh
+++ b/Engine/Extras/ushell/channels/flow/posix/provision.sh
@@ -7,7 +7,12 @@ _header()  { printf "\x1b[96m-- %s\x1b[0m\n" "$1" ; }
 _success() { printf "\x1b[92m%s\x1b[0m\n" "$1" ; }
 
 py_ver_maj=3
-py_ver_min=12
+if [ $(uname) == "Linux" ]; then
+    py_ver_min=11
+else
+    py_ver_min=12
+fi
+
 py_ver=$py_ver_maj.$py_ver_min
 
 # If we've already a destination we're done here
@@ -16,12 +21,44 @@ if [ -z "$1" ]; then
     exit 1
 fi
 
-# If we've already a destination we're done here
+py_mark=$py_ver.version
+
+# If current/ is already at the correct version, we're done
+if [ -e $1/python/current/$py_mark ]; then
+    exit 0
+fi
+
 dest_dir=$1/python/$py_ver
-if [ -d "$dest_dir" ]; then
+current_dir=$(dirname $dest_dir)/current
+function _link_dest_current()
+{
+    _header "Symlinking $dest_dir to $current_dir"
+    ln -s $dest_dir $current_dir
+    touch $dest_dir/$py_mark
+}
+
+# Remove any lingering current/
+if [ -d $current_dir ]; then
+    if ! rm $current_dir; then
+        _error "Failed unlinking $current_dir"
+        exit 1
+    fi
+fi
+
+# If we've already a destination we're alomst done here
+if [ -e $dest_dir/$py_mark ]; then
+    _link_dest_current
     exit 0
 fi
 
+# Remove any lingering prior venvs
+if [ -d $dest_dir ]; then
+    if ! rm -rf $dest_dir; then
+        _error "Failed removing $dest_dir"
+        exit 1
+    fi
+fi
+
 # Check the versions good enough
 function _check_version()
 {
@@ -29,7 +66,7 @@ function _check_version()
     cat << EOF | $1 -EsSB
 import sys
 vi = sys.version_info
-raise SystemExit(vi.major != ${py_ver_maj} or vi.minor < ${py_ver_min})
+raise SystemExit(vi.major != ${py_ver_maj} or vi.minor != ${py_ver_min})
 EOF
     return $?
 }
@@ -62,9 +99,10 @@ _bin_select \
     python
 
 if [ -z "$py_bin" ]; then
-    _error "No suitable binary found that is $py_ver or newer"
+    _error "No suitable binary found that is $py_ver"
     echo ""
     echo "Examples to install Python $py_ver:"
+    echo ""
     echo "   Ubuntu: apt install python$py_ver"
     echo "      Mac: brew install python@$py_ver"
     echo ""
@@ -73,28 +111,13 @@ fi
 
 _success "Found: $py_bin: (as $($py_bin --version))"
 
-# Make sure that the Python install has Pip installed
-_header "Checking for Pip"
-curl -L https://bootstrap.pypa.io/get-pip.py | $py_bin
-if [ "$?" -ne "0" ]; then
-    _error "Failed to ensure Pip is installed"
-    exit 1
-fi
-
-# Create the Python virtual env
-_header "Fetching virtual env"
-
+# Create the Python virtual env. Ubuntu doesn't come with the ensurepip module
+# which venv will run on start up. So our venv is created without Pip initially.
 temp_dir=$1/python/${py_ver}_$$
 
-if ! $py_bin -m pip install virtualenv; then
-    _error "Unable to install 'virtualenv' Pip package"
-    _error "Please ensure Pip is installed"
-    exit 1
-fi
-
-_header "Creating a virtualenv in $temp_dir"
-if ! $py_bin -m virtualenv $temp_dir; then
-    _error "Virtualenv creation failed"
+_header "Creating a virtual environment in $temp_dir"
+if ! $py_bin -m venv --without-pip $temp_dir; then
+    _error "venv-based virtual environment creation failed"
     exit 1
 fi
 
@@ -105,12 +128,14 @@ if [ "$?" -ne "0" ]; then
     exit 1
 fi
 
+# Make sure the venv has Pip and CA certificates
+_header "Fetching Pip"
+curl -L https://bootstrap.pypa.io/get-pip.py | $temp_dir/bin/python
+$temp_dir/bin/python -m pip install certifi
+
 # Swap the temp dir into place. If we fail someone else did it first
 if mv $temp_dir $dest_dir; then
-    current_dir=$(dirname $dest_dir)/current
-    _header "Symlinking $dest_dir to $current_dir"
-    rm $current_dir 2>/dev/null
-    ln -s $dest_dir $current_dir
+    _link_dest_current
 else
     rm -rf $temp_dir
 fi
diff --git a/Engine/Extras/ushell/channels/flow/posix/shells.py b/Engine/Extras/ushell/channels/flow/posix/shells.py
--- a/Engine/Extras/ushell/channels/flow/posix/shells.py
+++ b/Engine/Extras/ushell/channels/flow/posix/shells.py
@@ -47,7 +47,7 @@ class _Base(object):
                 value = str(value).replace("\n", "\\n")
                 out(f"export {key}=$'{value}'")
 
-        out("\$tip")
+        out(r"\$tip")
         out("echo")
 
 
@@ -77,18 +77,19 @@ class _Bash(_Base):
         working_dir = system.get_working_dir()
         script_path = working_dir + "complete.bash"
 
-        global _complete_bash
+        out_sh = _get_complete_bash()
         cmd_tree = system.get_command_tree()
         tree_root = cmd_tree.get_root_node()
         for name,_ in tree_root.read_children():
-            _complete_bash += " " + shlex.quote(name)
+            out_sh += " " + shlex.quote(name)
 
         with open(script_path, "w") as out:
-            out.write(_complete_bash)
+            out.write(out_sh)
 
         return script_path
 
-_complete_bash = r"""
+def _get_complete_bash():
+    return r"""
 function _flow_prompt() {
     if [[ ! -z "$FLOW_PROMPT" ]]; then
         local prompt=${FLOW_PROMPT/:FLOW_CWD:/$(dirs +0)}
@@ -146,15 +147,16 @@ class _Zsh(_Base):
         tree_root = cmd_tree.get_root_node()
         flow_cmds = "' '".join(name for name,_ in tree_root.read_children())
 
-        global _complete_zsh
-        out_sh = _complete_zsh.replace("$$$FLOW_CMDS$$$", flow_cmds)
+        out_sh = _get_complete_zsh()
+        out_sh = out_sh.replace("$$$FLOW_CMDS$$$", flow_cmds)
 
         with open(script_path, "w") as out:
             out.write(out_sh)
 
         return script_path
 
-_complete_zsh = r"""
+def _get_complete_zsh():
+    return r"""
 _flow_prompt() {
     FLOW_CWD=$(dirs)
     eval $(\$prompt --format=sh)
diff --git a/Engine/Extras/ushell/channels/unreal/core/cmds/run.py b/Engine/Extras/ushell/channels/unreal/core/cmds/run.py
--- a/Engine/Extras/ushell/channels/unreal/core/cmds/run.py
+++ b/Engine/Extras/ushell/channels/unreal/core/cmds/run.py
@@ -28,16 +28,19 @@ class _Attachable(object):
     attach = unrealcmd.Opt(False, "Attach a debugger to the launched process")
 
     def _get_debugger(self, name=None):
+        name = name or os.getenv("USHELL_DEBUGGER")
         name = name or ("vs" if os.name == "nt" else "lldb")
-        ue_context = self.get_unreal_context()
-        channel = self.get_channel()
-        debuggers = channel.get_extendable("debuggers")
+        dbg_py_path = Path(__file__).parent.parent / f"debuggers/{name}.py"
 
         try:
-            debugger_class = debuggers.get_extension_class(name)
-            return debugger_class.construct(ue_context)
-        except ValueError:
-            raise ValueError(f"Unknown debugger '{name}'. Valid list: [{','.join([x[0] for x in debuggers.read_extensions()])}]")
+            import importlib.util
+            spec = importlib.util.spec_from_file_location("", dbg_py_path)
+            module = importlib.util.module_from_spec(spec)
+            spec.loader.exec_module(module)
+            debugger_class = getattr(module, "Debugger", None)
+            return debugger_class(self.get_unreal_context())
+        except Exception as e:
+            raise RuntimeError(f"Failed to load debugger '{dbg_py_path}' ({e})")
 
     def attach_to(self, pid, *, transport=None, host_ip=None):
         debugger = self._get_debugger()
@@ -398,6 +401,7 @@ class Game(_Runtime):
         yield from super().complete_platform(prefix)
         yield "editor"
 
+    @unrealcmd.Cmd.summarise
     def main(self):
         if self.args.platform == "editor":
             editor = Editor()
@@ -412,6 +416,7 @@ class Game(_Runtime):
 #-------------------------------------------------------------------------------
 class Client(_Runtime):
     """ Launches the game client. """
+    @unrealcmd.Cmd.summarise
     def main(self):
         return self._main_impl(unreal.TargetType.CLIENT)
 
@@ -419,6 +424,7 @@ class Client(_Runtime):
 class Server(_Runtime):
     """ Launches the server binary. """
     complete_platform = ("win64", "linux")
+    @unrealcmd.Cmd.summarise
     def main(self):
         self.args.runargs = (*self.args.runargs, "-log", "-unattended")
         return self._main_impl(unreal.TargetType.SERVER)
diff --git a/Engine/Extras/ushell/channels/unreal/core/cmds/stage.py b/Engine/Extras/ushell/channels/unreal/core/cmds/stage.py
--- a/Engine/Extras/ushell/channels/unreal/core/cmds/stage.py
+++ b/Engine/Extras/ushell/channels/unreal/core/cmds/stage.py
@@ -109,6 +109,8 @@ class _Impl(unrealcmd.Cmd):
 
         platform = self.get_platform(platform)
 
+        cook_flavor = platform.get_cook_flavor()
+
         dot_uat_cmd = "_uat"
         dot_uat_args = (
             "BuildCookRun",
@@ -121,6 +123,7 @@ class _Impl(unrealcmd.Cmd):
             "-deploy" if self.args.deploy else None,
             "-config=" + self.args.variant,
             "-platform=" + platform.get_name(),
+            ("-cookflavor=" + cook_flavor) if cook_flavor else None,
             None if self.args.nopak else "-pak",
             *self.args.uatargs,
         )
diff --git a/Engine/Extras/ushell/channels/unreal/core/cmds/zen.py b/Engine/Extras/ushell/channels/unreal/core/cmds/zen.py
--- a/Engine/Extras/ushell/channels/unreal/core/cmds/zen.py
+++ b/Engine/Extras/ushell/channels/unreal/core/cmds/zen.py
@@ -6,155 +6,20 @@ import shutil
 import unreal
 import unrealcmd
 import unreal.cmdline
+import unreal.zencmd
 import http.client
 import json
 import hashlib
 
 #-------------------------------------------------------------------------------
-class _Base(unrealcmd.Cmd):
-    def perform_project_completion(self, prefix, hostname="localhost", port=8558):
-        headers = {"Accept":"application/json"}
-        conn = http.client.HTTPConnection(hostname, port=port, timeout=0.1)
-        projects = []
-        try:
-            conn.request("GET", "/prj", headers=headers)
-            response = conn.getresponse()
-            projects = json.loads(response.read().decode())
-        except socket.timeout:
-            return
-
-        for project in projects:
-            yield project['Id']
-
-    def perform_oplog_completion(self, prefix, project_filter, hostname="localhost", port=8558):
-        headers = {"Accept":"application/json"}
-        conn = http.client.HTTPConnection(hostname, port=port, timeout=0.1)
-        projects = []
-        try:
-            if project_filter:
-                conn.request("GET", f"/prj/{project_filter}", headers=headers)
-                response = conn.getresponse()
-                projects.append(json.loads(response.read().decode()))
-            else:
-                conn.request("GET", "/prj", headers=headers)
-                response = conn.getresponse()
-                projects = json.loads(response.read().decode())
-        except socket.timeout:
-            return
-
-        for project in projects:
-            for oplog in project['oplogs']:
-                yield oplog['id']
-
-    def get_project_id_from_project(self, project):
-        if not project:
-            self.print_error("An active project is required for this operation")
-
-        normalized_project_path = str(project.get_path()).replace("\\","/")
-        project_id = project.get_name() + "." + hashlib.md5(normalized_project_path.encode('utf-8')).hexdigest()[:8]
-        return project_id;
-
-
-    def get_project_or_default(self, explicit_project):
-        if explicit_project:
-            return explicit_project
-
-        ue_context = self.get_unreal_context()
-        return self.get_project_id_from_project(ue_context.get_project())
-
-
-#-------------------------------------------------------------------------------
-class _ZenUETargetBase(_Base):
-    variant = unrealcmd.Arg("development", "Build variant to launch")
-    runargs  = unrealcmd.Arg([str], "Additional arguments to pass to the process being launched")
-    build    = unrealcmd.Opt(True, "Build the target prior to running it")
-
-    def _build(self, target, variant, platform):
-        if target.get_type() == unreal.TargetType.EDITOR:
-            build_cmd = ("_build", "editor", variant)
-        else:
-            platform = platform or unreal.Platform.get_host()
-            build_cmd = ("_build", "target", target.get_name(), platform, variant)
-
-        import subprocess
-        ret = subprocess.run(build_cmd)
-        if ret.returncode:
-            raise SystemExit(ret.returncode)
-
-    def get_binary_path(self, target, platform=None):
-        ue_context = self.get_unreal_context()
-
-        if isinstance(target, str):
-            target = ue_context.get_target_by_name(target)
-        else:
-            target = ue_context.get_target_by_type(target)
-
-        if self.args.build:
-            self._build(target, self.args.variant, platform)
-
-        variant = unreal.Variant.parse(self.args.variant)
-        if build := target.get_build(variant, platform):
-            if build := build.get_binary_path():
-                return build
-
-        raise EnvironmentError(f"No {self.args.variant} build found for target '{target.get_name()}'")
-
-    def run_ue_program_target(self, target_name, cmdargs = None):
-        binary_path = self.get_binary_path(target_name)
-
-        final_args = []
-        if cmdargs: final_args += cmdargs
-        if self.args.runargs: final_args += (*unreal.cmdline.read_ueified(*self.args.runargs),)
-
-        launch_kwargs = {}
-        if not sys.stdout.isatty():
-            launch_kwargs = { "silent" : True }
-        exec_context = self.get_exec_context()
-        runnable = exec_context.create_runnable(str(binary_path), *final_args);
-        runnable.launch(**launch_kwargs)
-        return True if runnable.is_gui() else runnable.wait()
-
-
-#-------------------------------------------------------------------------------
-class _ZenUtilityBase(_Base):
-    zenargs  = unrealcmd.Arg([str], "Additional arguments to pass to zen utility")
-
-    def get_zen_utility_command(self, cmdargs):
-        ue_context = self.get_unreal_context()
-        engine = ue_context.get_engine()
-
-        platform = unreal.Platform.get_host()
-        exe_filename = "zen"
-        if platform == "Win64":
-            exe_filename = exe_filename + ".exe"
-
-        bin_path = engine.get_dir()
-        bin_path /= f"Binaries/{platform}/{exe_filename}"
-        if not bin_path.is_file():
-            raise FileNotFoundError(f"Unable to find '{bin_path}'")
-
-        final_args = []
-        if cmdargs: final_args += cmdargs
-        if self.args.zenargs: final_args += self.args.zenargs
-
-        return bin_path, final_args
-
-
-    def run_zen_utility(self, cmdargs):
-        bin_path, final_args = self.get_zen_utility_command(cmdargs)
-
-        cmd = self.get_exec_context().create_runnable(bin_path, *final_args)
-        return cmd.run()
-
-#-------------------------------------------------------------------------------
-class Dashboard(_ZenUETargetBase):
+class Dashboard(unreal.zencmd.ZenUETargetBaseCmd):
     """ Starts the zen dashboard GUI."""
 
     def main(self):
         return self.run_ue_program_target("ZenDashboard")
 
 #-------------------------------------------------------------------------------
-class Start(_ZenUETargetBase):
+class Start(unreal.zencmd.ZenUETargetBaseCmd):
     """ Starts an instance of zenserver."""
     SponsorProcessID = unrealcmd.Opt("", "Process ID to be added as a sponsor process for the zenserver process")
 
@@ -169,28 +34,28 @@ class Start(_ZenUETargetBase):
         return self.run_ue_program_target("ZenLaunch", args)
 
 #-------------------------------------------------------------------------------
-class Stop(_ZenUtilityBase):
+class Stop(unreal.zencmd.ZenUtilityBaseCmd):
     """ Stops any running instance of zenserver."""
 
     def main(self):
         return self.run_zen_utility(["down"])
 
 #-------------------------------------------------------------------------------
-class Status(_ZenUtilityBase):
+class Status(unreal.zencmd.ZenUtilityBaseCmd):
     """ Get the status of running zenserver instances."""
 
     def main(self):
         return self.run_zen_utility(["status"])
 
 #-------------------------------------------------------------------------------
-class Version(_ZenUtilityBase):
+class Version(unreal.zencmd.ZenUtilityBaseCmd):
     """ Get the version of the in-tree zenserver executable."""
 
     def main(self):
         return self.run_zen_utility(["version"])
 
 #-------------------------------------------------------------------------------
-class ImportSnapshot(_ZenUtilityBase):
+class ImportSnapshot(unreal.zencmd.ZenUtilityBaseCmd):
     """ Imports an oplog snapshot into the running zenserver process."""
     snapshotdescriptor = unrealcmd.Arg(str, "Snapshot descriptor file path to import from")
     snapshotindex = unrealcmd.Arg(0, "0-based index of the snapshot within the snapshot descriptor to import from")
diff --git a/Engine/Extras/ushell/channels/unreal/core/debuggers/lldb.py b/Engine/Extras/ushell/channels/unreal/core/debuggers/lldb.py
--- a/Engine/Extras/ushell/channels/unreal/core/debuggers/lldb.py
+++ b/Engine/Extras/ushell/channels/unreal/core/debuggers/lldb.py
@@ -2,12 +2,12 @@
 
 import os
 import signal
-import unreal
 import subprocess
 
 #-------------------------------------------------------------------------------
-class Debugger(unreal.Debugger):
-    name = "lldb"
+class Debugger(object):
+    def __init__(self, ue_context):
+        self._ue_context = ue_context
 
     def _disable_ctrl_c(self):
         def nop_handler(*args): pass
@@ -16,12 +16,12 @@ class Debugger(unreal.Debugger):
     def _restore_ctrl_c(self):
         signal.signal(signal.SIGINT, self._prev_sigint_handler)
 
-    def _debug(self, exec_context, cmd, *args):
+    def debug(self, exec_context, cmd, *args):
         self._disable_ctrl_c()
         subprocess.run(("lldb", cmd, "--", *args))
         self._restore_ctrl_c()
 
-    def _attach(self, pid, transport=None, host_ip=None):
+    def attach(self, pid, transport=None, host_ip=None):
         self._disable_ctrl_c()
         subprocess.run(("lldb", "--attach-pid", str(pid)))
         self._restore_ctrl_c()
diff --git a/Engine/Extras/ushell/channels/unreal/core/debuggers/vs.py b/Engine/Extras/ushell/channels/unreal/core/debuggers/vs.py
--- a/Engine/Extras/ushell/channels/unreal/core/debuggers/vs.py
+++ b/Engine/Extras/ushell/channels/unreal/core/debuggers/vs.py
@@ -1,28 +1,28 @@
 # Copyright Epic Games, Inc. All Rights Reserved.
 
 import os
-import unreal
 import subprocess
 from pathlib import Path
 
 #-------------------------------------------------------------------------------
-class Debugger(unreal.Debugger):
-    name = "Visual Studio"
+class Debugger(object):
+    def __init__(self, ue_context):
+        self._ue_context = ue_context
 
-    def _debug(self, exec_context, cmd, *args):
+    def debug(self, exec_context, cmd, *args):
         cmd = exec_context.create_runnable(cmd, *args)
         cmd.launch(suspended=True, new_term=True)
         pid = cmd.get_pid()
 
         try:
-            if not self._attach(pid):
+            if not self.attach(pid):
                 cmd.kill()
         except:
             cmd.kill()
             raise
 
-    def _attach(self, pid, transport=None, host_ip=None):
-        engine = self.get_unreal_context().get_engine()
+    def attach(self, pid, transport=None, host_ip=None):
+        engine = self._ue_context.get_engine()
         engine_dir = engine.get_dir()
 
         import vs.dte
diff --git a/Engine/Extras/ushell/channels/unreal/core/describe.flow.py b/Engine/Extras/ushell/channels/unreal/core/describe.flow.py
--- a/Engine/Extras/ushell/channels/unreal/core/describe.flow.py
+++ b/Engine/Extras/ushell/channels/unreal/core/describe.flow.py
@@ -6,14 +6,14 @@ import flow.describe
 #-------------------------------------------------------------------------------
 fzf_linux = flow.describe.Tool()
 fzf_linux.version("0.44.1")
-fzf_linux.payload("[REDACTED_PROFILE_URL]/fzf/releases/download/0.44.1/fzf-0.44.1-linux_amd64.tar.gz")
+fzf_linux.payload("[REDACTED_PROFILE_URL]/fzf/releases/download/$VERSION/fzf-$VERSION-linux_amd64.tar.gz")
 fzf_linux.sha1("e7780f1e5e2dd4c8aa5e94dcbcbedf4cebfaaceb")
 fzf_linux.platform("linux")
 fzf_linux.bin("fzf")
 
 fzf_darwin = flow.describe.Tool()
 fzf_darwin.version("0.44.1")
-fzf_darwin.payload("[REDACTED_PROFILE_URL]/fzf/releases/download/0.44.1/fzf-0.44.1-darwin_amd64.zip")
+fzf_darwin.payload("[REDACTED_PROFILE_URL]/fzf/releases/download/$VERSION/fzf-$VERSION-darwin_amd64.zip")
 fzf_darwin.sha1("405156eb6fc3d5275774cc569e24feaed362eab5")
 fzf_darwin.platform("darwin")
 fzf_darwin.bin("fzf")
@@ -23,7 +23,7 @@ fzf_darwin.bin("fzf")
 #-------------------------------------------------------------------------------
 ripgrep_win32 = flow.describe.Tool()
 ripgrep_win32.version("14.0.3")
-ripgrep_win32.payload("[REDACTED_PROFILE_URL]/ripgrep/releases/download/14.0.3/ripgrep-14.0.3-x86_64-pc-windows-msvc.zip")
+ripgrep_win32.payload("[REDACTED_PROFILE_URL]/ripgrep/releases/download/$VERSION/ripgrep-$VERSION-x86_64-pc-windows-msvc.zip")
 ripgrep_win32.sha1("5508b3dd5c12206c5d1c7994490bacf317cee0aa")
 ripgrep_win32.platform("win32")
 ripgrep_win32.bin("rg.exe")
@@ -31,14 +31,14 @@ ripgrep_win32.source("[REDACTED_PROFILE_URL]/ripgrep/releases/latest", r"
 
 ripgrep_linux = flow.describe.Tool()
 ripgrep_linux.version("14.0.3")
-ripgrep_linux.payload("[REDACTED_PROFILE_URL]/ripgrep/releases/download/14.0.3/ripgrep-14.0.3-x86_64-unknown-linux-musl.tar.gz")
+ripgrep_linux.payload("[REDACTED_PROFILE_URL]/ripgrep/releases/download/$VERSION/ripgrep-$VERSION-x86_64-unknown-linux-musl.tar.gz")
 ripgrep_linux.sha1("25a759834819f91625ac739c890a4b0139df7b63")
 ripgrep_linux.platform("linux")
 ripgrep_linux.bin("rg")
 
 ripgrep_darwin = flow.describe.Tool()
 ripgrep_darwin.version("14.0.3")
-ripgrep_darwin.payload("[REDACTED_PROFILE_URL]/ripgrep/releases/download/14.0.3/ripgrep-14.0.3-x86_64-apple-darwin.tar.gz")
+ripgrep_darwin.payload("[REDACTED_PROFILE_URL]/ripgrep/releases/download/$VERSION/ripgrep-$VERSION-x86_64-apple-darwin.tar.gz")
 ripgrep_darwin.sha1("d5af393ea9b5a443544edc686cf939f517be35ea")
 ripgrep_darwin.platform("darwin")
 ripgrep_darwin.bin("rg")
@@ -242,8 +242,8 @@ tips.prefix("$")
 
 #-------------------------------------------------------------------------------
 vswhere = flow.describe.Tool()
-vswhere.version("3.0.1")
-vswhere.payload("[REDACTED_PROFILE_URL]/vswhere/releases/download/3.1.7/vswhere.exe")
+vswhere.version("3.1.7")
+vswhere.payload("[REDACTED_PROFILE_URL]/vswhere/releases/download/$VERSION/vswhere.exe")
 vswhere.sha1("e3fa9b2db259d8875170717469779ea1280c8466")
 vswhere.platform("win32")
 vswhere.bin("vswhere.exe")
@@ -251,20 +251,6 @@ vswhere.source("[REDACTED_PROFILE_URL]/vswhere/releases", r"(\d+\.\d+\.\d+
 
 
 
-#-------------------------------------------------------------------------------
-def _add_debugger(name):
-    extension = flow.describe.Extension()
-    extension.source(f"debuggers/{name}.py", "Debugger")
-    extension.mount("debuggers", name)
-    globals()[name] = extension
-
-_add_debugger("lldb")
-if sys.platform == "win32":
-    _add_debugger("vs")
-
-
-
 #-------------------------------------------------------------------------------
 unreal = flow.describe.Channel()
 unreal.version("0")
-unreal.extendable("debuggers")
diff --git a/Engine/Extras/ushell/channels/unreal/core/pylib/unreal/__init__.py b/Engine/Extras/ushell/channels/unreal/core/pylib/unreal/__init__.py
--- a/Engine/Extras/ushell/channels/unreal/core/pylib/unreal/__init__.py
+++ b/Engine/Extras/ushell/channels/unreal/core/pylib/unreal/__init__.py
@@ -4,4 +4,3 @@ from ._context import *
 from ._ini import *
 from ._platform import *
 from ._ubt import *
-from ._debugger import *
diff --git a/Engine/Extras/ushell/channels/unreal/core/pylib/unreal/_context.py b/Engine/Extras/ushell/channels/unreal/core/pylib/unreal/_context.py
--- a/Engine/Extras/ushell/channels/unreal/core/pylib/unreal/_context.py
+++ b/Engine/Extras/ushell/channels/unreal/core/pylib/unreal/_context.py
@@ -622,7 +622,7 @@ class Context(object):
         iter_dir:Path = uproj_path.parent.parent
         while not iter_dir.samefile(iter_dir.anchor):
             candidate = iter_dir / "Engine"
-            if candidate.is_dir():
+            if candidate.is_dir() and (candidate / "Config").is_dir():
                 self._type = ContextType.PROJECT
                 self._engine_dir = candidate
                 return
diff --git a/Engine/Extras/ushell/channels/unreal/core/pylib/unreal/_debugger.py b/Engine/Extras/ushell/channels/unreal/core/pylib/unreal/_debugger.py
deleted file mode 100644
diff --git a/Engine/Extras/ushell/channels/unreal/core/pylib/unreal/_platform.py b/Engine/Extras/ushell/channels/unreal/core/pylib/unreal/_platform.py
--- a/Engine/Extras/ushell/channels/unreal/core/pylib/unreal/_platform.py
+++ b/Engine/Extras/ushell/channels/unreal/core/pylib/unreal/_platform.py
@@ -78,17 +78,17 @@ class Platform(object):
             # As the time of writing the SDK.json files have trailing commas
             # and C-style comments. So we'll try and conform it
             ret = ""
-            count = 0
             for line in inp:
+                try: line = line[:line.index("//")]
+                except ValueError: pass
                 line = line.strip()
-                if not line or line.startswith("//"):
+                if not line:
                     continue
 
-                count += 1
-                sep = "," if count > 1 else ""
-                if line.endswith(","):    line = line[:-1]
-                if line[0] in "[{":       count = 0
-                elif line[0] in "}]":     sep = ""
+                line = line[:-1] if line.endswith(",") else line
+                sep = "" if ret[-1:] in "[{" else ","
+                sep = "" if line[0] in "]}" else sep
+
                 ret += sep + line
             return ret
 
@@ -153,6 +153,10 @@ class Platform(object):
             return form
         raise ValueError(f"No known cook form for target '{target}'")
 
+    def get_cook_flavor(self):
+        if hasattr(self, "_get_cook_flavor"):
+            return self._get_cook_flavor()
+
     @_optional_api
     def launch(self, exec_context, stage_dir, binary_path, args):
         return self._launch(exec_context, stage_dir, binary_path, args)
diff --git a/Engine/Extras/ushell/channels/unreal/core/pylib/unreal/_ubt.py b/Engine/Extras/ushell/channels/unreal/core/pylib/unreal/_ubt.py
--- a/Engine/Extras/ushell/channels/unreal/core/pylib/unreal/_ubt.py
+++ b/Engine/Extras/ushell/channels/unreal/core/pylib/unreal/_ubt.py
@@ -219,8 +219,10 @@ class _BuildToolBase(object):
         suffix = "UnrealBuildTool/BuildConfiguration.xml"
         yield _BuildConfig(engine_dir / "Restricted/NotForLicensees/Programs" / suffix, _BuildConfigLevel.INTERNAL)
         yield self.get_configuration(True)
-        yield self.get_configuration()
         if os.name == "nt":
+            for var_name in ("ProgramData", "APPDATA", "LOCALAPPDATA"):
+                prefix = os.getenv(var_name, "!notset")
+                yield _BuildConfig(prefix + "/Unreal Engine/" + suffix, _BuildConfigLevel.GLOBAL)
             user_profile = os.getenv("USERPROFILE", "userprofile-notset")
             yield _BuildConfig(user_profile + "/Documents/Unreal Engine/" + suffix, _BuildConfigLevel.LEGACY)
         else:
@@ -233,7 +235,7 @@ class _BuildToolBase(object):
         if branch:
             cfg_path = self._engine.get_dir() / "Saved" / suffix
         else:
-            if os.name == "nt": prefix = os.getenv("APPDATA", "appdata")
+            if os.name == "nt": prefix = os.getenv("LOCALAPPDATA", "appdata")
             else:               prefix = os.path.expanduser("~")
             cfg_path = Path(prefix) / "Unreal Engine" / suffix
 
diff --git a/Engine/Extras/ushell/channels/unreal/core/pylib/unreal/platforms/win64/android.py b/Engine/Extras/ushell/channels/unreal/core/pylib/unreal/platforms/win64/android.py
--- a/Engine/Extras/ushell/channels/unreal/core/pylib/unreal/platforms/win64/android.py
+++ b/Engine/Extras/ushell/channels/unreal/core/pylib/unreal/platforms/win64/android.py
@@ -103,6 +103,9 @@ class Platform(unreal.Platform):
         if target == "game":   return "Android_ASTC"
         if target == "client": return "Android_ASTCClient"
 
+    def _get_cook_flavor(self):
+        return "ASTC"
+
     def _get_package_name(self):
         config = self.get_unreal_context().get_config()
         if ret := str(config.get("Engine", Platform.config_section, "PackageName")):
diff --git a/Engine/Extras/ushell/channels/unreal/core/pylib/unreal/zencmd.py b/Engine/Extras/ushell/channels/unreal/core/pylib/unreal/zencmd.py
new file mode 100644
--- /dev/null
+++ b/Engine/Extras/ushell/channels/unreal/core/pylib/unreal/zencmd.py
@@ -0,0 +1,147 @@
+# Copyright Epic Games, Inc. All Rights Reserved.
+
+import os
+import sys
+import shutil
+import unreal
+import unrealcmd
+import unreal.cmdline
+import http.client
+import json
+import hashlib
+
+#-------------------------------------------------------------------------------
+class ZenBaseCmd(unrealcmd.Cmd):
+    def perform_project_completion(self, prefix, hostname="localhost", port=8558):
+        headers = {"Accept":"application/json"}
+        conn = http.client.HTTPConnection(hostname, port=port, timeout=0.1)
+        projects = []
+        try:
+            conn.request("GET", "/prj", headers=headers)
+            response = conn.getresponse()
+            projects = json.loads(response.read().decode())
+        except socket.timeout:
+            return
+
+        for project in projects:
+            yield project['Id']
+
+    def perform_oplog_completion(self, prefix, project_filter, hostname="localhost", port=8558):
+        headers = {"Accept":"application/json"}
+        conn = http.client.HTTPConnection(hostname, port=port, timeout=0.1)
+        projects = []
+        try:
+            if project_filter:
+                conn.request("GET", f"/prj/{project_filter}", headers=headers)
+                response = conn.getresponse()
+                projects.append(json.loads(response.read().decode()))
+            else:
+                conn.request("GET", "/prj", headers=headers)
+                response = conn.getresponse()
+                projects = json.loads(response.read().decode())
+        except socket.timeout:
+            return
+
+        for project in projects:
+            for oplog in project['oplogs']:
+                yield oplog['id']
+
+    def get_project_id_from_project(self, project):
+        if not project:
+            self.print_error("An active project is required for this operation")
+
+        normalized_project_path = str(project.get_path()).replace("\\","/")
+        project_id = project.get_name() + "." + hashlib.md5(normalized_project_path.encode('utf-8')).hexdigest()[:8]
+        return project_id;
+
+
+    def get_project_or_default(self, explicit_project):
+        if explicit_project:
+            return explicit_project
+
+        ue_context = self.get_unreal_context()
+        return self.get_project_id_from_project(ue_context.get_project())
+
+
+#-------------------------------------------------------------------------------
+class ZenUETargetBaseCmd(ZenBaseCmd):
+    variant = unrealcmd.Arg("development", "Build variant to launch")
+    runargs  = unrealcmd.Arg([str], "Additional arguments to pass to the process being launched")
+    build    = unrealcmd.Opt(True, "Build the target prior to running it")
+
+    def _build(self, target, variant, platform):
+        if target.get_type() == unreal.TargetType.EDITOR:
+            build_cmd = ("_build", "editor", variant)
+        else:
+            platform = platform or unreal.Platform.get_host()
+            build_cmd = ("_build", "target", target.get_name(), platform, variant)
+
+        import subprocess
+        ret = subprocess.run(build_cmd)
+        if ret.returncode:
+            raise SystemExit(ret.returncode)
+
+    def get_binary_path(self, target, platform=None):
+        ue_context = self.get_unreal_context()
+
+        if isinstance(target, str):
+            target = ue_context.get_target_by_name(target)
+        else:
+            target = ue_context.get_target_by_type(target)
+
+        if self.args.build:
+            self._build(target, self.args.variant, platform)
+
+        variant = unreal.Variant.parse(self.args.variant)
+        if build := target.get_build(variant, platform):
+            if build := build.get_binary_path():
+                return build
+
+        raise EnvironmentError(f"No {self.args.variant} build found for target '{target.get_name()}'")
+
+    def run_ue_program_target(self, target_name, cmdargs = None):
+        binary_path = self.get_binary_path(target_name)
+
+        final_args = []
+        if cmdargs: final_args += cmdargs
+        if self.args.runargs: final_args += (*unreal.cmdline.read_ueified(*self.args.runargs),)
+
+        launch_kwargs = {}
+        if not sys.stdout.isatty():
+            launch_kwargs = { "silent" : True }
+        exec_context = self.get_exec_context()
+        runnable = exec_context.create_runnable(str(binary_path), *final_args);
+        runnable.launch(**launch_kwargs)
+        return True if runnable.is_gui() else runnable.wait()
+
+
+#-------------------------------------------------------------------------------
+class ZenUtilityBaseCmd(ZenBaseCmd):
+    zenargs  = unrealcmd.Arg([str], "Additional arguments to pass to zen utility")
+
+    def get_zen_utility_command(self, cmdargs):
+        ue_context = self.get_unreal_context()
+        engine = ue_context.get_engine()
+
+        platform = unreal.Platform.get_host()
+        exe_filename = "zen"
+        if platform == "Win64":
+            exe_filename = exe_filename + ".exe"
+
+        bin_path = engine.get_dir()
+        bin_path /= f"Binaries/{platform}/{exe_filename}"
+        if not bin_path.is_file():
+            raise FileNotFoundError(f"Unable to find '{bin_path}'")
+
+        final_args = []
+        if cmdargs: final_args += cmdargs
+        if self.args.zenargs: final_args += self.args.zenargs
+
+        return bin_path, final_args
+
+
+    def run_zen_utility(self, cmdargs):
+        bin_path, final_args = self.get_zen_utility_command(cmdargs)
+
+        cmd = self.get_exec_context().create_runnable(bin_path, *final_args)
+        return cmd.run()
\ No newline at end of file
diff --git a/Engine/Extras/ushell/channels/unreal/perforce/cmds/project.py b/Engine/Extras/ushell/channels/unreal/perforce/cmds/project.py
--- a/Engine/Extras/ushell/channels/unreal/perforce/cmds/project.py
+++ b/Engine/Extras/ushell/channels/unreal/perforce/cmds/project.py
@@ -3,6 +3,7 @@
 import os
 import p4utils
 import flow.cmd
+from pathlib import Path
 
 #-------------------------------------------------------------------------------
 class Change(flow.cmd.Cmd):
@@ -12,8 +13,21 @@ class Change(flow.cmd.Cmd):
         if not primary_dir:
             return False
 
+        # Don't do anything if a .git/ or a previous failure are detected
+        for item in Path(primary_dir).parents:
+            if (item / ".git").is_dir():            return True
+            if (item / ".p4.off.ushell").is_file():   return True
+            if (item / "GenerateProjectFiles.bat").is_file():
+                primary_dir = item
+                break
+        else:
+            return True
+
+        # Help the user by ensure there is a P4CONFIG file
         self.print_info("Perforce environment")
         try:
+            import time
+            ts = time.perf_counter()
             p4config = p4utils.ensure_p4config(primary_dir)
             if p4config:
                 path, created = p4config
@@ -21,6 +35,12 @@ class Change(flow.cmd.Cmd):
                 print("Path:", path)
             else:
                 self.print_warning("Unabled to find a client from", primary_dir)
+                if time.perf_counter() - ts > 5.0:
+                    self.print_warning()
+                    self.print_warning("Perforce environment checks by '.project' are now disabled")
+                    self.print_warning("Delete '<branch_root>/.p4.off.ushell' to reinstate")
+                    with (primary_dir / ".p4.off.ushell").open("wt") as out:
+                        out.write(".project failed detecting Perforce env")
         except EnvironmentError as e:
             self.print_warning(e)
 
diff --git a/Engine/Extras/ushell/ushell.sh b/Engine/Extras/ushell/ushell.sh
--- a/Engine/Extras/ushell/ushell.sh
+++ b/Engine/Extras/ushell/ushell.sh
@@ -2,8 +2,8 @@
 
 # Check whether this script has been "sourced" into a host shell
 if [ $0 = "$BASH_SOURCE" ]; then
-    
-    # The script was not sourced, so source it in a child shell and propagate any arguments
+    # The script was not sourced, so source it in a child shell and propagate
+    # any arguments
     args=""
     for arg in "$@"
     do
@@ -12,7 +12,6 @@ if [ $0 = "$BASH_SOURCE" ]; then
     done
     bash --init-file <(echo "source \"$HOME/.bashrc\"; source \"$BASH_SOURCE\" $args")
     exit $?
-    
 fi
 
 # Determine the host shell
