Updated ushell [32394647]

#rnx

diff --git a/Engine/Extras/ushell/channels/flow/nt/provision.bat b/Engine/Extras/ushell/channels/flow/nt/provision.bat
--- a/Engine/Extras/ushell/channels/flow/nt/provision.bat
+++ b/Engine/Extras/ushell/channels/flow/nt/provision.bat
@@ -7,13 +7,12 @@ goto:eof
 ::------------------------------------------------------------------------------
 :main
 
-set _pyver=3.12.1
-set _pysha=7a77c7471f5ceee88831a29b9460eddd784adf31b0230f188149ef0c588642d1
+set _pyver=3.12.2
+set _pysha=144ee3fa0c4171f98ec554ac01bb44ee1e895d154609f143b1d20b52c3fe383e
 set _pytag=312
+set _pymark=%_pyver%.version
 
-set _destdir=%~f1\%_pyver%
-
-if exist "%_destdir%" (
+if exist "%~f1\current\%_pymark%" (
     goto:eof
 )
 
@@ -21,6 +20,8 @@ call:check_bin tar.exe _tar_path
 call:check_bin curl.exe _curl_path
 call:check_bin certutil.exe _certutil_path
 
+set _destdir=%~f1\%_pyver%
+
 1>nul 2>nul (
     rd "%_destdir%\..\current"
 )
@@ -106,6 +107,7 @@ call:get_url https://bootstrap.pypa.io/get-pip.py
 call:on_error "Failed to add Pip"
 del get-pip.py
 
+>%_pymark% echo %_pyver%
 popd
 goto:eof
 
diff --git a/Engine/Extras/ushell/channels/flow/posix/shells.py b/Engine/Extras/ushell/channels/flow/posix/shells.py
--- a/Engine/Extras/ushell/channels/flow/posix/shells.py
+++ b/Engine/Extras/ushell/channels/flow/posix/shells.py
@@ -168,6 +168,10 @@ fi
 if [[ -z "$FLOW_CMDS" ]]; then
     precmd_functions+=(_flow_prompt)
 fi
-unset FLOW_CMDS
-FLOW_CMDS+=('$$$FLOW_CMDS$$$')
+
+function _flow_complete() {
+    # setopt local_options xtrace
+    compadd $('$complete' -- ${^^words[@]:0:-1} $PREFIX...)
+}
+compdef _flow_complete '$$$FLOW_CMDS$$$'
 """
