PR #11490: Resolves an issue where the current working directory for the complet?

#rnx
#11490

diff --git a/Engine/Extras/ushell/channels/flow/core/cmds/complete.py b/Engine/Extras/ushell/channels/flow/core/cmds/complete.py
--- a/Engine/Extras/ushell/channels/flow/core/cmds/complete.py
+++ b/Engine/Extras/ushell/channels/flow/core/cmds/complete.py
@@ -109,6 +109,7 @@ def _complete(system, word_iter, reply, delim):
         return True
 
     if words and words[0] == "$$":
+        os.chdir(os.path.normpath(words[1]))
         _get_tags(system, reply)
         return True
 
diff --git a/Engine/Extras/ushell/channels/flow/nt/shell_pwsh.py b/Engine/Extras/ushell/channels/flow/nt/shell_pwsh.py
--- a/Engine/Extras/ushell/channels/flow/nt/shell_pwsh.py
+++ b/Engine/Extras/ushell/channels/flow/nt/shell_pwsh.py
@@ -181,6 +181,7 @@ function Update-UShellEnvVars {{
 
     $CompleteDaemon.StandardInput.WriteLine("`u{{01}}")
     $CompleteDaemon.StandardInput.WriteLine("`$`$")
+    $CompleteDaemon.StandardInput.WriteLine($($executionContext.SessionState.Path.CurrentLocation))
     $CompleteDaemon.StandardInput.WriteLine("`u{{02}}")
 
     $Results = @{{}}
