PR #11525: Improved help print for ushell. Trims docstrings as per PEP257, after which text?

#11525
#rnx

diff --git a/Engine/Extras/ushell/channels/flow/core/cmds/help.py b/Engine/Extras/ushell/channels/flow/core/cmds/help.py
--- a/Engine/Extras/ushell/channels/flow/core/cmds/help.py
+++ b/Engine/Extras/ushell/channels/flow/core/cmds/help.py
@@ -36,7 +36,7 @@ class Help(flow.cmd.Cmd):
                     command_class = node.get_command_class()
                     command_class = command_class.get_class_type()
                     if hasattr(command_class, "get_desc"):
-                        yield path, command_class.get_desc().strip()
+                        yield path, command_class.get_desc()
                 except Exception as e:
                     yield path, "LOAD_ERROR: " + str(e)
 
@@ -47,23 +47,56 @@ class Help(flow.cmd.Cmd):
         output.sort(key=lambda x: x[0])
 
         name_len = max(len(x[0]) for x in output) + 4
-        whitespace = re.compile(r"\s+")
         desc_width = right_column - name_len
         for name, desc in output:
             dots = "." * (name_len - len(name) - 3)
             initial_prefix = flow.cmd.text.light_cyan(name) + " " + flow.cmd.text.grey(dots)
 
-            desc = whitespace.sub(" ", desc)
-            desc_lines = textwrap.wrap(desc, desc_width)
+            is_first_paragraph = True
 
-            first_line = 0
-            for line in desc_lines:
-                prefix = initial_prefix if not first_line else (" " * (name_len - 2))
-                first_line += 1
-                print("", prefix, line)
+            paragraph_lines = []
+            is_literal_block = None
 
-            if first_line > 1:
-                print()
+            def print_paragraph(is_literal_block, is_first_paragraph):
+                if paragraph_lines:
+                    # literal blocks retain line endings and wrap on indent
+                    if is_literal_block:
+                        desc_lines = []
+                        for line in paragraph_lines:
+                            line_len = len(line)
+                            line = line.lstrip()
+                            line_indent = line_len - len(line)
+
+                            desc_lines.extend([(" " * line_indent) + l for l in textwrap.wrap(line, desc_width)])
+                    else:
+                        desc_lines = textwrap.wrap(" ".join(paragraph_lines), desc_width)
+
+                    paragraph_prefix = (" " * (name_len - 2))
+
+                    prefix = initial_prefix if is_first_paragraph else paragraph_prefix
+                    print("", prefix, desc_lines[0])
+
+                    for line in desc_lines[1:]:
+                        print("", paragraph_prefix, line)
+                    
+                    del paragraph_lines[:]
+
+                    if desc_lines:
+                        print()
+
+            for desc_line in desc.splitlines():
+                if not desc_line:
+                    print_paragraph(is_literal_block, is_first_paragraph)
+                    is_first_paragraph = False
+                    is_literal_block = None
+                else:
+                    # treat paragraph as a literal block if it starts out indented
+                    if is_literal_block is None:
+                        is_literal_block = desc_line[0].isspace()
+
+                    paragraph_lines.append(desc_line)
+
+            print_paragraph(is_literal_block, is_first_paragraph)
 
 
 
diff --git a/Engine/Extras/ushell/channels/flow/core/system/flow/_flick.py b/Engine/Extras/ushell/channels/flow/core/system/flow/_flick.py
--- a/Engine/Extras/ushell/channels/flow/core/system/flow/_flick.py
+++ b/Engine/Extras/ushell/channels/flow/core/system/flow/_flick.py
@@ -1,5 +1,42 @@
 # Copyright Epic Games, Inc. All Rights Reserved.
 
+import sys
+
+#-------------------------------------------------------------------------------
+def trim(docstring):
+    """Function from PEP257 to handle docstring indentation.
+    
+    Strips a uniform amount of indentation from the second and further lines of
+    the docstring, equal to the minimum indentation of all non-blank lines
+    after the first line. Any indentation in the first line of the docstring
+    (i.e., up to the first newline) is insignificant and removed. Relative
+    indentation of later lines in the docstring is retained. Blank lines should
+    be removed from the beginning and end of the docstring.
+    """
+    if not docstring:
+        return ''
+    # Convert tabs to spaces (following the normal Python rules)
+    # and split into a list of lines:
+    lines = docstring.expandtabs().splitlines()
+    # Determine minimum indentation (first line doesn't count):
+    indent = sys.maxsize
+    for line in lines[1:]:
+        stripped = line.lstrip()
+        if stripped:
+            indent = min(indent, len(line) - len(stripped))
+    # Remove indentation (first line is special):
+    trimmed = [lines[0].strip()]
+    if indent < sys.maxsize:
+        for line in lines[1:]:
+            trimmed.append(line[indent:].rstrip())
+    # Strip off trailing and leading blank lines:
+    while trimmed and not trimmed[-1]:
+        trimmed.pop()
+    while trimmed and not trimmed[0]:
+        trimmed.pop(0)
+    # Return a single string:
+    return '\n'.join(trimmed)
+
 #-------------------------------------------------------------------------------
 class _ArgOptBase(object):
     def __init__(self, type_value, description):
@@ -289,7 +326,7 @@ class Cmd(object):
 
     @classmethod
     def get_desc(self):
-        it = (x.__doc__ for x in self.__mro__ if x.__doc__ and x != object)
+        it = (trim(x.__doc__) for x in self.__mro__ if x.__doc__ and x != object)
         return "\n\n".join(it)
 
     def invoke(self, in_args):
