Corrected errors in README.txt and offered additional clarity.

#rnx

diff --git a/Engine/Extras/ushell/README.txt b/Engine/Extras/ushell/README.txt
--- a/Engine/Extras/ushell/README.txt
+++ b/Engine/Extras/ushell/README.txt
@@ -127,23 +127,39 @@ running in a non-interactive scripting context (see `Scripting`)
 # Alternative Host Shells
 
 A lot of the above assumes that cmd.exe is the host shell. Other shells are also
-supported; Bash or Zsh on POSIX-based platforms, and PowerShell. There are
-scripts in ushell's root folder for using these alternative shells.
+supported; Bash or Zsh on POSIX-based platforms (see "Mac and Linux" section),
+and PowerShell. There are scripts in ushell's root folder for using different
+host shells which advanced users can reference to support more exotic hosts.
 
 ## PowerShell
 
-PowerShell integration works by importing ushell as a module. Set your
-`PSModulePath` to contain the ushell's directory (this can be done in your
-`$PROFILE` file);
+PowerShell integration works by importing ushell as a module. This is typically
+done via the `$PROFILE` script. Consider the case where ushell is located in
+`c:\path\to\ushell`. First add the parent directory to `PSModulePath` (notice
+the absence of `ushell\` here);
+
+```
+$env:PSModulePath += ";c:\path\to\"
+```
+
+You can then can import the module, either in your profile script or when you
+want to add ushell to a particular terminal session;
 
 ```
-$env:PSModulePath = "$($env:PSModulePath);c:\path\to\ushell\"
 Import-Module ushell
 ```
 
+It is important to note that the ushell folder must be called "ushell" as this
+is required by PowerShell's module import mechanism. If this is not the case you
+can also import modules by absolute path;
+
+```
+Import-Module c:\path\to\ushell\ushell.psm1
+```
+
 If you use a PowerShell prompt enhancer like oh-my-posh you can extract
 environment variables from ushell to populate your prompt/window title/etc. The
-ushell-related environment variables can be list with `dir env:USHELL*`.
+ushell-related environment variables can be listed with `dir env:USHELL*`.
 
 ```
 function Set-MyPoshContext {
