Updated ushell [35015634]

#rnx

diff --git a/Engine/Extras/ushell/README.txt b/Engine/Extras/ushell/README.txt
--- a/Engine/Extras/ushell/README.txt
+++ b/Engine/Extras/ushell/README.txt
@@ -127,39 +127,32 @@ running in a non-interactive scripting context (see `Scripting`)
 # Alternative Host Shells
 
 A lot of the above assumes that cmd.exe is the host shell. Other shells are also
-supported; Bash or Zsh on POSIX-based platforms (see "Mac and Linux" section),
-and PowerShell. There are scripts in ushell's root folder for using different
-host shells which advanced users can reference to support more exotic hosts.
+supported; Bash or Zsh on POSIX-based platforms, and PowerShell. There are
+scripts in ushell's root folder for using these alternative shells.
 
 ## PowerShell
 
-PowerShell integration works by importing ushell as a module. This is typically
-done via the `$PROFILE` script. Consider the case where ushell is located in
-`c:\path\to\ushell`. First add the parent directory to `PSModulePath` (notice
-the absence of `ushell\` here);
+PowerShell integration works by importing ushell as a module. Powershell wants
+modules to be contained within a directory matching the name of the module, so the
+'powerushell' module is contained in the path ushell\powerushell\powerushell.psm1.
 
-```
-$env:PSModulePath += ";c:\path\to\"
-```
-
-You can then can import the module, either in your profile script or when you
-want to add ushell to a particular terminal session;
+Set your `PSModulePath` to contain the ushell directory (this can be done in your
+`$PROFILE` file).
 
 ```
-Import-Module ushell
+$env:PSModulePath = "$($env:PSModulePath);c:\path\to\ushell\"
 ```
 
-It is important to note that the ushell folder must be called "ushell" as this
-is required by PowerShell's module import mechanism. If this is not the case you
-can also import modules by absolute path;
+Then you can import the module, either in your profile or when you want to enter a
+ushell session for a particular terminal session.
 
 ```
-Import-Module c:\path\to\ushell\ushell.psm1
+Import-Module powerushell
 ```
 
 If you use a PowerShell prompt enhancer like oh-my-posh you can extract
 environment variables from ushell to populate your prompt/window title/etc. The
-ushell-related environment variables can be listed with `dir env:USHELL*`.
+ushell-related environment variables can be list with `dir env:USHELL*`.
 
 ```
 function Set-MyPoshContext {
diff --git a/Engine/Extras/ushell/channels/flow/core/boot.py b/Engine/Extras/ushell/channels/flow/core/boot.py
--- a/Engine/Extras/ushell/channels/flow/core/boot.py
+++ b/Engine/Extras/ushell/channels/flow/core/boot.py
@@ -48,7 +48,7 @@ class Boot(flow.cmd.Cmd):
         os.environ["PATH"] = env["PATH"]
 
         # Set a session id and somewhere to store per-session information
-        session_id = str(os.getpid())
+        session_id = str(os.getppid())
         env["FLOW_SID"] = session_id
         os.environ["FLOW_SID"] = session_id
 
@@ -66,9 +66,12 @@ class Boot(flow.cmd.Cmd):
         except KeyError:
             raise RuntimeError(f"Unable to find a shell named '{shell_name}'")
 
+        # This is a bit crude, but some shell support allows for customisation
+        user_script = getattr(self, "_user_script", None)
+
         # Run the shell
         system = channel.get_system()
-        shell_class(system).boot_shell(env, shell_cookie)
+        shell_class(system).boot_shell(env, shell_cookie, user_script)
 
         # Apply a theme
         if self.args.theme:
diff --git a/Engine/Extras/ushell/channels/flow/core/cmds/help.py b/Engine/Extras/ushell/channels/flow/core/cmds/help.py
--- a/Engine/Extras/ushell/channels/flow/core/cmds/help.py
+++ b/Engine/Extras/ushell/channels/flow/core/cmds/help.py
@@ -1,9 +1,6 @@
 # Copyright Epic Games, Inc. All Rights Reserved.
 
-import re
-import sys
 import shutil
-import marshal
 import flow.cmd
 import textwrap
 import importlib
@@ -35,8 +32,7 @@ class Help(flow.cmd.Cmd):
                 try:
                     command_class = node.get_command_class()
                     command_class = command_class.get_class_type()
-                    if hasattr(command_class, "get_desc"):
-                        yield path, command_class.get_desc()
+                    yield path, command_class
                 except Exception as e:
                     yield path, "LOAD_ERROR: " + str(e)
 
@@ -47,56 +43,17 @@ class Help(flow.cmd.Cmd):
         output.sort(key=lambda x: x[0])
 
         name_len = max(len(x[0]) for x in output) + 4
+        lead_dot = "." * (name_len - 2)
+        lead_spc = " " * name_len
         desc_width = right_column - name_len
-        for name, desc in output:
-            dots = "." * (name_len - len(name) - 3)
-            initial_prefix = flow.cmd.text.light_cyan(name) + " " + flow.cmd.text.grey(dots)
-
-            is_first_paragraph = True
-
-            paragraph_lines = []
-            is_literal_block = None
-
-            def print_paragraph(is_literal_block, is_first_paragraph):
-                if paragraph_lines:
-                    # literal blocks retain line endings and wrap on indent
-                    if is_literal_block:
-                        desc_lines = []
-                        for line in paragraph_lines:
-                            line_len = len(line)
-                            line = line.lstrip()
-                            line_indent = line_len - len(line)
-
-                            desc_lines.extend([(" " * line_indent) + l for l in textwrap.wrap(line, desc_width)])
-                    else:
-                        desc_lines = textwrap.wrap(" ".join(paragraph_lines), desc_width)
-
-                    paragraph_prefix = (" " * (name_len - 2))
-
-                    prefix = initial_prefix if is_first_paragraph else paragraph_prefix
-                    print("", prefix, desc_lines[0])
-
-                    for line in desc_lines[1:]:
-                        print("", paragraph_prefix, line)
-                    
-                    del paragraph_lines[:]
-
-                    if desc_lines:
-                        print()
-
-            for desc_line in desc.splitlines():
-                if not desc_line:
-                    print_paragraph(is_literal_block, is_first_paragraph)
-                    is_first_paragraph = False
-                    is_literal_block = None
-                else:
-                    # treat paragraph as a literal block if it starts out indented
-                    if is_literal_block is None:
-                        is_literal_block = desc_line[0].isspace()
-
-                    paragraph_lines.append(desc_line)
-
-            print_paragraph(is_literal_block, is_first_paragraph)
+        for name, cmd_class in output:
+            dots = lead_dot[len(name):]
+            leader = flow.cmd.text.light_cyan(name) + " " + flow.cmd.text.grey(dots)
+            desc = cmd_class.get_desc(pretty=desc_width)
+            desc = desc.replace("\n", "\n" + lead_spc)
+            print(" ", leader, sep="", end="")
+            print(desc)
+            print()
 
 
 
diff --git a/Engine/Extras/ushell/channels/flow/core/system/flow/_flick.py b/Engine/Extras/ushell/channels/flow/core/system/flow/_flick.py
--- a/Engine/Extras/ushell/channels/flow/core/system/flow/_flick.py
+++ b/Engine/Extras/ushell/channels/flow/core/system/flow/_flick.py
@@ -1,42 +1,5 @@
 # Copyright Epic Games, Inc. All Rights Reserved.
 
-import sys
-
-#-------------------------------------------------------------------------------
-def trim(docstring):
-    """Function from PEP257 to handle docstring indentation.
-    
-    Strips a uniform amount of indentation from the second and further lines of
-    the docstring, equal to the minimum indentation of all non-blank lines
-    after the first line. Any indentation in the first line of the docstring
-    (i.e., up to the first newline) is insignificant and removed. Relative
-    indentation of later lines in the docstring is retained. Blank lines should
-    be removed from the beginning and end of the docstring.
-    """
-    if not docstring:
-        return ''
-    # Convert tabs to spaces (following the normal Python rules)
-    # and split into a list of lines:
-    lines = docstring.expandtabs().splitlines()
-    # Determine minimum indentation (first line doesn't count):
-    indent = sys.maxsize
-    for line in lines[1:]:
-        stripped = line.lstrip()
-        if stripped:
-            indent = min(indent, len(line) - len(stripped))
-    # Remove indentation (first line is special):
-    trimmed = [lines[0].strip()]
-    if indent < sys.maxsize:
-        for line in lines[1:]:
-            trimmed.append(line[indent:].rstrip())
-    # Strip off trailing and leading blank lines:
-    while trimmed and not trimmed[-1]:
-        trimmed.pop()
-    while trimmed and not trimmed[0]:
-        trimmed.pop(0)
-    # Return a single string:
-    return '\n'.join(trimmed)
-
 #-------------------------------------------------------------------------------
 class _ArgOptBase(object):
     def __init__(self, type_value, description):
@@ -173,9 +136,7 @@ class Cmd(object):
     def _print_help(self, short=False):
         # Description
         if not short:
-            for line in self.get_desc().split("\n"):
-                line = line[:4].strip() + line[4:]
-                print(line)
+            print(self.get_desc(pretty=80))
             print()
 
         # Usage
@@ -325,9 +286,41 @@ class Cmd(object):
         return self.__name__
 
     @classmethod
-    def get_desc(self):
-        it = (trim(x.__doc__) for x in self.__mro__ if x.__doc__ and x != object)
-        return "\n\n".join(it)
+    def get_desc(self, *, pretty=False):
+        it = (x.__doc__ for x in self.__mro__ if x.__doc__ and x != object)
+        if not int(pretty):
+            return "\n\n".join(it)
+
+        desc = []
+        for lines in it:
+            lines = lines.strip().splitlines()
+            it = (len(x) - len(x.lstrip()) for x in lines[1:] if x.lstrip())
+            to_trim = min(it, default=0)
+            desc.append(lines[0])
+            desc += list(x[to_trim:] for x in lines[1:])
+            desc.append("")
+
+        if type(pretty) != int:
+            return "\n".join(desc[:-1])
+
+        import textwrap
+        leadless_n = 0
+        paras = []
+        for line in desc:
+            next_n = (leadless_n + 1) if line and line[0] != " " else 0
+            if next_n > 1:
+                paras[-1] += " " + line
+                leadless_n = next_n
+                continue
+
+            if next_n < leadless_n:
+                para = paras.pop()
+                paras += textwrap.wrap(para, width=pretty)
+
+            paras.append(line)
+            leadless_n = next_n
+
+        return "\n".join(paras[:-1])
 
     def invoke(self, in_args):
         self.validate()
diff --git a/Engine/Extras/ushell/channels/flow/nt/cmd_thefuzz.bat b/Engine/Extras/ushell/channels/flow/nt/cmd_thefuzz.bat
--- a/Engine/Extras/ushell/channels/flow/nt/cmd_thefuzz.bat
+++ b/Engine/Extras/ushell/channels/flow/nt/cmd_thefuzz.bat
@@ -8,7 +8,8 @@ set _fzf=fzf --layout=reverse --color=16 --height=20
 
 <nul >%_bat% (
     if "%~1"=="history" call:history
-    if "%~1"=="chdir"   call:chdir %2
+    if "%~1"=="chdir"   call:chdir "cd /d" %2
+    if "%~1"=="pushdir" call:chdir pushd %2
     if "%~1"=="explore" call:explore %2
 )
 
@@ -49,9 +50,9 @@ goto:eof
 
 ::------------------------------------------------------------------------------
 :chdir
-set _base="%~f1"
-if "%~1"=="" set _base=\
-echo cd /d ^^
+set _base="%~f2"
+if "%~2"=="" set _base=\
+echo %~1 ^^
 fd -d8 -td . %_base:\"="% | %_fzf% --scheme=default --history=%temp%/thefuzz_chdir
 goto:eof
 
diff --git a/Engine/Extras/ushell/channels/flow/nt/describe.flow.py b/Engine/Extras/ushell/channels/flow/nt/describe.flow.py
--- a/Engine/Extras/ushell/channels/flow/nt/describe.flow.py
+++ b/Engine/Extras/ushell/channels/flow/nt/describe.flow.py
@@ -1,6 +1,5 @@
 # Copyright Epic Games, Inc. All Rights Reserved.
 
-import os
 import flow.describe
 
 #-------------------------------------------------------------------------------
diff --git a/Engine/Extras/ushell/channels/flow/nt/provision.bat b/Engine/Extras/ushell/channels/flow/nt/provision.bat
--- a/Engine/Extras/ushell/channels/flow/nt/provision.bat
+++ b/Engine/Extras/ushell/channels/flow/nt/provision.bat
@@ -12,7 +12,11 @@ set _pysha=144ee3fa0c4171f98ec554ac01bb44ee1e895d154609f143b1d20b52c3fe383e
 set _pytag=312
 set _pymark=%_pyver%.version
 
-if exist "%~f1\current\%_pymark%" (
+:: temp folder for detritus that's rmdir'd at the end
+set _tempdir=%~f1\$del
+
+if exist "%~f1\current\%_pymark%" 1>nul 2>nul (
+    rd /q /s "%_tempdir%"
     goto:eof
 )
 
@@ -20,16 +24,21 @@ call:check_bin tar.exe _tar_path
 call:check_bin curl.exe _curl_path
 call:check_bin certutil.exe _certutil_path
 
-set _destdir=%~f1\%_pyver%
+1>nul 2>nul (
+    mkdir "%_tempdir%"
+)
 
 1>nul 2>nul (
-    rd "%_destdir%\..\current"
+    :: as current/python.exe maybe locked, we move instead of rmdir
+    move "%~f1\current" "%_tempdir%\cur_%random%"
 )
 
+set _destdir=%_tempdir%\%_pyver%
 call:get_python "%_destdir%"
 
 1>nul 2>nul (
-    mklink /j "%_destdir%\..\current" "%_destdir%"
+    move "%_destdir%" "%~f1\current"
+    rd /q /s "%_tempdir%"
 )
 
 goto:eof
diff --git a/Engine/Extras/ushell/channels/flow/nt/shell_cmd.py b/Engine/Extras/ushell/channels/flow/nt/shell_cmd.py
--- a/Engine/Extras/ushell/channels/flow/nt/shell_cmd.py
+++ b/Engine/Extras/ushell/channels/flow/nt/shell_cmd.py
@@ -6,10 +6,54 @@
 from pathlib import Path
 
 #-------------------------------------------------------------------------------
+class Cmd(object):
+    def register_shells(self, registrar):
+        registrar.add("cmd", _Shell)
+        return super().register_shells(registrar)
+
+    def run(self, env):
+        if os.name != "nt":
+            return super().run(env)
+
+        # Give the user a convenient way to add customisations in to the session
+        self.print_info("Customisation")
+        home_dir = Path(self.get_home_dir())
+        tags = (
+            flow.cmd.text.grey("none"),
+            flow.cmd.text.white("okay"),
+        )
+
+        # These are really a no-op. It is here to hint to the user how they can
+        # add their own channels
+        candidate = home_dir / "channels"
+        found = candidate.is_dir()
+        print(tags[int(found)], candidate / "*")
+
+        # Indicate to the user that there's a way to hook the boot process too
+        candidate = home_dir / "hooks/boot.bat"
+        found = candidate.is_file()
+        print(tags[int(found)], candidate)
+
+        # This is run just before the prompt is shown
+        candidate = home_dir / "hooks/startup.bat"
+        found = candidate.is_file()
+        print(tags[int(found)], candidate)
+        if found:
+            self._user_script = candidate
+
+        print()
+        return super().run(env)
+
+#-------------------------------------------------------------------------------
 class _Shell(object):
+    def __init__(self, system):
+        self._system = system
+
+    def get_system(self):
+        return self._system
+
     def _add_clink_lua(self):
-        channel = self.get_channel()
-        system = channel.get_system()
+        system = self.get_system()
         working_dir = system.get_working_dir()
 
         lua_path = working_dir + "complete.clink.lua"
@@ -34,15 +78,10 @@
         lua.write(clink_lua)
         lua.close()
 
-    def register_shells(self, registrar):
-        registrar.add("cmd", lambda x: self)
-        return super().register_shells(registrar)
-
-    def boot_shell(self, env, cookie):
+    def boot_shell(self, env, cookie, user_script):
         self._add_clink_lua()
 
-        channel = self.get_channel()
-        system = channel.get_system()
+        system = self.get_system()
         working_dir = system.get_working_dir()
         env["CLINK_PATH"] = working_dir
         env["DIRCMD"] = "/ogen"
@@ -68,49 +107,16 @@
             print(rf'doskey de="{thefuzz_path}" explore $*', file=out)
             print(rf'doskey dc="{thefuzz_path}" chdir $*', file=out)
             print(rf'doskey dc.="{thefuzz_path}" chdir .', file=out)
+            print(rf'doskey dp="{thefuzz_path}" pushdir $*', file=out)
+            print(rf'doskey dp.="{thefuzz_path}" pushdir .', file=out)
 
             print("clink_x64.exe inject", file=out)
             print("$tip.exe", file=out)
 
-            if user_script := getattr(self, "_user_script", None):
+            if user_script:
                 print(rf'call "{user_script}"', file=out)
 
 
-#-------------------------------------------------------------------------------
-class Cmd(_Shell):
-    def run(self, env):
-        if os.name != "nt":
-            return super().run(env)
-
-        # Give the user a convenient way to add customisations in to the session
-        self.print_info("Customisation")
-        home_dir = Path(self.get_home_dir())
-        tags = (
-            flow.cmd.text.grey("none"),
-            flow.cmd.text.white("okay"),
-        )
-
-        # These are really a no-op. It is here to hint to the user how they can
-        # add their own channels
-        candidate = home_dir / "channels"
-        found = candidate.is_dir()
-        print(tags[int(found)], candidate / "*")
-
-        # Indicate to the user that there's a way to hook the boot process too
-        candidate = home_dir / "hooks/boot.bat"
-        found = candidate.is_file()
-        print(tags[int(found)], candidate)
-
-        # This is run just before the prompt is shown
-        candidate = home_dir / "hooks/startup.bat"
-        found = candidate.is_file()
-        print(tags[int(found)], candidate)
-        if found:
-            self._user_script = candidate
-
-        print()
-        return super().run(env)
-
 
 #-------------------------------------------------------------------------------
 def _get_clink_lua():
diff --git a/Engine/Extras/ushell/channels/flow/nt/shell_pwsh.py b/Engine/Extras/ushell/channels/flow/nt/shell_pwsh.py
--- a/Engine/Extras/ushell/channels/flow/nt/shell_pwsh.py
+++ b/Engine/Extras/ushell/channels/flow/nt/shell_pwsh.py
@@ -9,6 +9,7 @@ class Pwsh(object):
         registrar.add("pwsh", _Shell)
         return super().register_shells(registrar)
 
+#-------------------------------------------------------------------------------
 class _Shell(object):
     def __init__(self, system):
         self._system = system
@@ -16,7 +17,7 @@ class _Shell(object):
     def get_system(self):
         return self._system
 
-    def boot_shell(self, env, cookie):
+    def boot_shell(self, env, cookie, user_script):
         # no prompt modification here for now
         # could try and modify the Prompt function from the cookie
         # could also publish environment variables users can use from their own prompts
diff --git a/Engine/Extras/ushell/channels/flow/posix/shells.py b/Engine/Extras/ushell/channels/flow/posix/shells.py
--- a/Engine/Extras/ushell/channels/flow/posix/shells.py
+++ b/Engine/Extras/ushell/channels/flow/posix/shells.py
@@ -11,8 +11,6 @@ class Shells(object):
         registrar.add("zsh", _Zsh)
         return super().register_shells(registrar)
 
-
-
 #-------------------------------------------------------------------------------
 class _Base(object):
     def __init__(self, system):
@@ -21,7 +19,7 @@ class _Base(object):
     def get_system(self):
         return self._system
 
-    def boot_shell(self, env, cookie):
+    def boot_shell(self, env, cookie, user_script):
         try: os.makedirs(os.path.dirname(cookie))
         except: pass
 
@@ -54,7 +52,7 @@ class _Base(object):
 
 #-------------------------------------------------------------------------------
 class _Bash(_Base):
-    def boot_shell(self, env, cookie):
+    def boot_shell(self, env, cookie, user_script):
         # Add '\[...\]' around ANSI codes so Bash knows what's displayable
         if prompt := env.get("FLOW_PROMPT", None):
             out_prompt = ""
@@ -66,7 +64,7 @@ class _Bash(_Base):
             out_prompt += prompt[prev:]
             env["FLOW_PROMPT"] = out_prompt
 
-        return super().boot_shell(env, cookie)
+        return super().boot_shell(env, cookie, user_script)
 
     def write_cookie(self, env, out):
         super().write_cookie(env, out)
@@ -110,7 +108,7 @@ complete -F _flow_complete""" # purposely left no trailing CRLF
 
 #-------------------------------------------------------------------------------
 class _Zsh(_Base):
-    def boot_shell(self, env, cookie):
+    def boot_shell(self, env, cookie, user_script):
         if prompt := env.get("FLOW_PROMPT", None):
             # Add '%{...%}' around ANSI codes so Zsh knows what's displayable
             out_prompt = ""
@@ -132,7 +130,7 @@ class _Zsh(_Base):
 
             env["FLOW_PROMPT"] = prompt
 
-        return super().boot_shell(env, cookie)
+        return super().boot_shell(env, cookie, user_script)
 
     def write_cookie(self, env, out):
         super().write_cookie(env, out)
diff --git a/Engine/Extras/ushell/channels/unreal/core/cmds/build.py b/Engine/Extras/ushell/channels/unreal/core/cmds/build.py
--- a/Engine/Extras/ushell/channels/unreal/core/cmds/build.py
+++ b/Engine/Extras/ushell/channels/unreal/core/cmds/build.py
@@ -272,6 +272,13 @@ class _BuildCmd(unrealcmd.Cmd):
                         file_or_module = str(file.resolve())
                         break
 
+        # Lets not accept anything accept files that make translation units
+        for ext in (".cpp", ".c", ".xpp", ".h"):
+            if file_or_module.endswith(ext):
+                break
+        else:
+            raise ValueError(f"{file_or_module} does not appear to be a source file")
+
         builder.add_args("-SingleFile=" + file_or_module)
         builder.add_args("-SkipDeploy")
 
diff --git a/Engine/Extras/ushell/channels/unreal/core/cmds/cook.py b/Engine/Extras/ushell/channels/unreal/core/cmds/cook.py
--- a/Engine/Extras/ushell/channels/unreal/core/cmds/cook.py
+++ b/Engine/Extras/ushell/channels/unreal/core/cmds/cook.py
@@ -24,14 +24,15 @@ class _CookPrettyPrinter(uelogprinter.Printer):
 
 #-------------------------------------------------------------------------------
 class _Cook(unrealcmd.MultiPlatformCmd):
-    cookargs = unrealcmd.Arg([str], "Additional arguments passed to the cook commandlet")
-    cultures = unrealcmd.Opt("en", "Cultures to cook (comma separated, defaults to 'en')")
-    onthefly = unrealcmd.Opt(False, "Launch as an on-the-fly server")
-    iterate  = unrealcmd.Opt(False, "Cook iteratively on top of the previous cook")
-    noxge    = unrealcmd.Opt(False, "Disable XGE-based shader compilation")
-    unpretty = unrealcmd.Opt(False, "Turns off colourful pretty-printing")
-    attach   = unrealcmd.Opt(False, "Attach a debugger to the cook")
-    debug    = unrealcmd.Opt(False, "Use debug executables")
+    cookargs     = unrealcmd.Arg([str], "Additional arguments passed to the cook commandlet")
+    cultures     = unrealcmd.Opt("en", "Cultures to cook (comma separated, defaults to 'en')")
+    onthefly     = unrealcmd.Opt(False, "Launch as an on-the-fly server")
+    iterate      = unrealcmd.Opt(False, "Cook iteratively on top of the previous cook")
+    noxge        = unrealcmd.Opt(False, "Disable XGE-based shader compilation")
+    unpretty     = unrealcmd.Opt(False, "Turns off colourful pretty-printing")
+    attach       = unrealcmd.Opt(False, "Attach a debugger to the cook")
+    debug        = unrealcmd.Opt(False, "Use debug executables")
+    nounattended = unrealcmd.Opt(False, "Skip implicitly specified '-unattended'")
 
     @unrealcmd.Cmd.summarise
     def _cook(self, cook_form, exec_context=None):
@@ -49,7 +50,7 @@ class _Cook(unrealcmd.MultiPlatformCmd):
         commandlet_args = (
             ("-targetplatform=" + cook_form) if cook_form else None,
             ("-cookcultures=" + cultures) if cultures else None,
-            "-unattended",
+            "-unattended" if not self.args.nounattended else None,
             "-unversioned",
             "-stdout",
             "-cookonthefly" if self.args.onthefly else None,
diff --git a/Engine/Extras/ushell/channels/unreal/core/cmds/odsc.py b/Engine/Extras/ushell/channels/unreal/core/cmds/odsc.py
new file mode 100644
--- /dev/null
+++ b/Engine/Extras/ushell/channels/unreal/core/cmds/odsc.py
@@ -0,0 +1,154 @@
+# Copyright Epic Games, Inc. All Rights Reserved.
+
+import re
+import unreal
+import flow.cmd
+import unrealcmd
+import uelogprinter
+import unreal.cmdline
+
+#-------------------------------------------------------------------------------
+class _OdscPrettyPrinter(uelogprinter.Printer):
+    def __init__(self):
+        super().__init__()
+        self._odsc_re = re.compile(r"ages (\d+).+otal (\d+)")
+
+    def _print(self, line):
+        m = self._odsc_re.search(line)
+        if m:
+            percent = (int(m.group(1)) * 100) // max(int(m.group(2)), 1)
+            line += f" [{percent}%]"
+        super()._print(line)
+
+
+
+#-------------------------------------------------------------------------------
+class _Odsc(unrealcmd.MultiPlatformCmd):
+    odscargs = unrealcmd.Arg([str], "Additional arguments passed to the odsc (cook) commandlet")
+    cultures = unrealcmd.Opt("en", "Cultures to cook (comma separated, defaults to 'en')")
+    noxge    = unrealcmd.Opt(False, "Disable XGE-based shader compilation")
+    unpretty = unrealcmd.Opt(False, "Turns off colourful pretty-printing")
+    attach   = unrealcmd.Opt(False, "Attach a debugger to the commandlet")
+    skipassetscan = unrealcmd.Opt(False, "Skip asset scanning")
+    debug    = unrealcmd.Opt(False, "Use debug executables")
+
+    @unrealcmd.Cmd.summarise
+    def _odsc(self, cook_form, exec_context=None):
+        ue_context = self.get_unreal_context()
+
+        # Check there is a valid project
+        project = ue_context.get_project()
+        if not project:
+            self.print_error("An active project is required to cook")
+            return False
+
+        # Prepare arguments
+        cultures = self.args.cultures.replace(",", "+")
+
+        commandlet_args = (
+            "-odsc",
+            ("-targetplatform=" + cook_form) if cook_form else None,
+            ("-cookcultures=" + cultures) if cultures else None,
+            "-unversioned",
+            "-stdout",
+            "-cookonthefly",
+            "-dpcvars=r.ShaderDevelopmentMode=1",
+            "-noxgeshadercompile" if self.args.noxge else None,
+            "-skipassetscan" if self.args.skipassetscan else None,
+            *unreal.cmdline.read_ueified(*self.args.odscargs),
+        )
+
+        commandlet_args = tuple(x for x in commandlet_args if x)
+
+        # To support attaching a debugger we will reuse '.run commandlet'. This
+        # approach is not used normally as it interferes with pretty-printing and
+        # creates an awful lot of child processes.
+
+        if self.args.attach:
+            args = ("_run", "commandlet", "cook", "--noplatforms", "--attach")
+            if self.args.debug:
+                args = (*args, "debug")
+            args = (*args, "--", *commandlet_args)
+
+            # Disable our handling of Ctrl-C
+            import signal
+            def nop_handler(*args): pass
+            signal.signal(signal.SIGINT, nop_handler)
+
+            import subprocess
+            return subprocess.run(args).returncode
+
+        # Find the editor binary
+        target = ue_context.get_target_by_type(unreal.TargetType.EDITOR)
+        variant = unreal.Variant.DEBUG if self.args.debug else unreal.Variant.DEVELOPMENT
+        build = target.get_build(variant=variant)
+        if not build:
+            self.print_error(f"No {variant.name.lower()} editor build found")
+            return False
+
+        cmd = str(build.get_binary_path())
+        if not cmd.endswith("-Cmd.exe"):
+            cmd = cmd.replace(".exe", "-Cmd.exe")
+
+        # Launch
+        args = (
+            project.get_path(),
+            "-run=cook",
+            *commandlet_args,
+        )
+
+        exec_context = exec_context or self.get_exec_context()
+        cmd = exec_context.create_runnable(cmd, *args)
+        if self.args.unpretty or not self.is_interactive():
+            cmd.run()
+        else:
+            printer = _OdscPrettyPrinter()
+            printer.run(cmd)
+        return cmd.get_return_code()
+
+#-------------------------------------------------------------------------------
+class _Runtime(_Odsc):
+    platform = unrealcmd.Arg("", "The platform to setup odsc for")
+
+    def main(self, target):
+        platform = self.args.platform
+        if not platform:
+            platform = unreal.Platform.get_host()
+        platform = self.get_platform(platform)
+
+        exec_context = super().get_exec_context()
+        exec_env = exec_context.get_env()
+
+        # Establish target platform's environment
+        name = platform.get_name()
+        self.print_info(name, "-", platform.get_version())
+        for item in platform.read_env():
+            try:
+                dir = str(item)
+                exec_env[item.key] = dir
+            except EnvironmentError:
+                dir = flow.cmd.text.red(item.get())
+            print(item.key, "=", dir)
+
+        cook_form = platform.get_cook_form(target.name.lower())
+        return self._odsc(cook_form, exec_context)
+
+#-------------------------------------------------------------------------------
+class Game(_Runtime):
+    """ Launches an On Demand Shader Compilation server. Please launch the cooked game with -odschost=[ODSC-SERVER-IP] """
+    def main(self):
+        return super().main(unreal.TargetType.GAME)
+
+#-------------------------------------------------------------------------------
+class Client(_Runtime):
+    """ Launches an On Demand Shader Compilation server. Please launch the cooked client with -odschost=[ODSC-SERVER-IP] """
+    def main(self):
+        return super().main(unreal.TargetType.CLIENT)
+
+#-------------------------------------------------------------------------------
+class All(_Odsc):
+    """ Launches an On Demand Shader Compilation server. Please launch the cooked client with -odschost=[ODSC-SERVER-IP] """
+    def main(self):
+        exec_context = super().get_exec_context()
+        exec_env = exec_context.get_env()
+        return self._odsc("All", exec_context)
diff --git a/Engine/Extras/ushell/channels/unreal/core/cmds/project.py b/Engine/Extras/ushell/channels/unreal/core/cmds/project.py
--- a/Engine/Extras/ushell/channels/unreal/core/cmds/project.py
+++ b/Engine/Extras/ushell/channels/unreal/core/cmds/project.py
@@ -14,6 +14,7 @@ class Change(flow.cmd.Cmd):
      .project branch              - unset the active project on the branch
      .project cwd                 - active project or branch follows CWD
      .project active              - output current active project
+     .project list                - list all projects available but don't switch
      .project                     - pick from a list of all projects under the CWD
     """
     nameorpath   = flow.cmd.Arg("", "Name or path of the project to make active")
@@ -37,6 +38,7 @@ class Change(flow.cmd.Cmd):
         yield "branch"
         yield "cwd"
         yield "active"
+        yield "list"
 
         session = self.get_noticeboard(self.Noticeboard.SESSION)
         context = unreal.Context(session["uproject"] or ".")
@@ -89,6 +91,15 @@ class Change(flow.cmd.Cmd):
             reply = reply[1].lstrip()
             return os.path.abspath(reply)
 
+    def _print_all(self):
+        context = unreal.Context(".")
+        branch = context.get_branch()
+        if not branch:
+            return
+        cwd = str(branch.get_dir())
+        for uproj_path in branch.read_projects():
+            print(f"{uproj_path.stem:19}\t{uproj_path}")
+
     def _print_active(self):
         session = self.get_noticeboard(self.Noticeboard.SESSION)
         context = unreal.Context(session["uproject"] or ".")
@@ -207,7 +218,6 @@ class Change(flow.cmd.Cmd):
 
     def main(self):
         ret = self._main_impl()
-        self._provision(ret)
         return ret
 
     def _main_impl(self):
@@ -221,6 +231,9 @@ class Change(flow.cmd.Cmd):
             self.print_error("No project name or path specified")
             return False
 
+        if candidate.lower() == "list":
+            return self._print_all()
+
         if candidate.lower() == "active":
             return self._print_active()
 
@@ -258,67 +271,3 @@ class Change(flow.cmd.Cmd):
         session["uproject"] = primary_dir
 
         return primary_dir
-
-    def _provision(self, primary_dir):
-        if not primary_dir:
-            return
-
-        from pathlib import Path
-
-        # Find branch root
-        for candidate in Path(__file__).parents:
-            if (candidate / "ushell.bat").is_file():
-                root_dir = candidate
-                break
-        else:
-            return
-
-        # Find the bundle of data to provision
-        input_path = root_dir / Path("channels/bundle")
-        if not input_path.is_file():
-            return
-
-        with input_path.open("rb") as inp:
-            input_data = memoryview(inp.read())
-
-        input_mtime = input_path.stat().st_mtime
-
-        # Check we've something to provision from
-        nyckels = self._get_nyckels(primary_dir)
-        if not nyckels:
-            return
-
-        self.print_info("Provisioning")
-
-        # Let's provision!
-        import cipher
-        blob = cipher.Blob(input_data)
-        for nyckel in nyckels:
-            for name, data in blob.find(nyckel):
-                path = root_dir / name
-                if not path.is_file():
-                    print(name)
-                    dest_path = (root_dir / path)
-                    dest_path.parent.mkdir(parents=True, exist_ok=True)
-                    with dest_path.open("wb") as out:
-                        out.write(data)
-
-    def _get_nyckels(self, primary_dir):
-        pass
-        """
-        context = unreal.Context(primary_dir)
-        ini_paths = [x for x in context.glob("Extras/ushell/Settings.ini")]
-        if not ini_paths:
-            return
-
-        # And we need to know what we can provision
-        import base64
-        nyckels = []
-        for ini_path in ini_paths:
-            ini = unreal.Ini()
-            with ini_path.open("rt") as inp:
-                ini.load(inp)
-            if nyckel := ini.ushell.blob:
-                nyckel = base64.b64decode(str(nyckel))
-                nyckels.append(nyckel)
-        """
diff --git a/Engine/Extras/ushell/channels/unreal/core/cmds/run.py b/Engine/Extras/ushell/channels/unreal/core/cmds/run.py
--- a/Engine/Extras/ushell/channels/unreal/core/cmds/run.py
+++ b/Engine/Extras/ushell/channels/unreal/core/cmds/run.py
@@ -9,17 +9,19 @@ from pathlib import Path
 
 
 #-------------------------------------------------------------------------------
-def _get_ip():
+def _get_ip(ver=4):
     import socket
-    s = socket.socket(socket.AF_INET, socket.SOCK_DGRAM)
-    try:
-        s.connect(("172.31.255.255", 1))
-        ret = s.getsockname()[0]
-    except:
-        ret = "127.0.0.1"
-    finally:
-        s.close()
-    return ret
+    def impl(family, ip_addr):
+        with socket.socket(family, socket.SOCK_DGRAM) as s:
+            try:
+                ip_addr = "172.31.255.255" if ver == 4 else "fe80::0"
+                s.connect((ip_addr, 1))
+                return s.getsockname()[0]
+            except:
+                return "127.0.0.1" if "." in ip_addr else "::1"
+    assert ver in (4, 6)
+    args = (socket.AF_INET, "172.31.255.255") if ver == 4 else (socket.AF_INET6, "fe80::0")
+    return impl(*args)
 
 
 
@@ -355,7 +357,9 @@ class _Runtime(_RunCmd, _Attachable):
         )
 
         if not self.args.cooked:
-            args = (*args, "-pak")
+            projectstore_marker = stage_dir / "ue.projectstore"
+            if not projectstore_marker.is_file():
+                args = (*args, "-pak")
 
         if self.args.trace is not False:
             args = (*args, "-tracehost=" + _get_ip())
diff --git a/Engine/Extras/ushell/channels/unreal/core/cmds/sln.py b/Engine/Extras/ushell/channels/unreal/core/cmds/sln.py
--- a/Engine/Extras/ushell/channels/unreal/core/cmds/sln.py
+++ b/Engine/Extras/ushell/channels/unreal/core/cmds/sln.py
@@ -131,3 +131,35 @@ class Open(_Base):
     def main(self):
         self.use_all_platforms()
         return self._open_sln()
+
+#-------------------------------------------------------------------------------
+class Open10x(_Base):
+    """ Opens a Visual Studio Solution in 10x Editor """
+
+    def main(self):
+        if os.name != "nt":
+            self.print_error("Opening Visual Studio Solution in 10x Editor is only supported on Windows")
+            return False
+
+        self.print_info("Opening solution in 10x Editor")
+        sln_path = self._get_sln_path()
+        print("Path:", sln_path)
+        if not os.path.isfile(sln_path):
+            self.print_error("Project file not found")
+            return False
+
+        # At the moment there is no env var set for 10x, but Stewart will be adding this for us
+        # Assume default install location for now.
+        editor_path: str = r"C:\Program Files\PureDevSoftware\10x\10x.exe"
+        if not os.path.isfile(sln_path):
+            self.print_error("10x.exe not found at default install location")
+            return False
+
+        args = (
+            sln_path,
+        )
+
+        cmd = self.get_exec_context().create_runnable(editor_path, *args)
+        result = cmd.launch()
+        return True
+
diff --git a/Engine/Extras/ushell/channels/unreal/core/cmds/stage.py b/Engine/Extras/ushell/channels/unreal/core/cmds/stage.py
--- a/Engine/Extras/ushell/channels/unreal/core/cmds/stage.py
+++ b/Engine/Extras/ushell/channels/unreal/core/cmds/stage.py
@@ -71,8 +71,8 @@ class _Impl(unrealcmd.Cmd):
         platform = self.args.platform = platform.lower()
 
         # Validate the given target
-        target = unreal.TargetType.parse(self.args.target)
-        target = ue_context.get_target_by_type(target)
+        target_type = unreal.TargetType.parse(self.args.target)
+        target = ue_context.get_target_by_type(target_type)
         if not target:
             self.print_error(f"'{project.get_name()}' does not appear to have a '{self.args.target}' target")
             return 1
@@ -109,6 +109,14 @@ class _Impl(unrealcmd.Cmd):
 
         platform = self.get_platform(platform)
 
+        if self.args.zen:
+            cook_form = platform.get_cook_form(target_type.name.lower())
+            projectstore_file = project.get_dir() / f"Saved/Cooked/{cook_form}/" / "ue.projectstore"
+            if not projectstore_file.is_file():
+                self.print_error(f"'{projectstore_file}' missing, cannot use --zen for '{self.args.target}' target")
+                return 1
+            self.args.nopak = True
+
         cook_flavor = platform.get_cook_flavor()
 
         dot_uat_cmd = "_uat"
@@ -148,6 +156,7 @@ class Stage(_Impl):
     cook   = unrealcmd.Opt((False, ""), "Run a cook before running the stage step")
     deploy = unrealcmd.Opt(False, "Deploy to devkit/device after staging")
     nopak  = unrealcmd.Opt(False, "Staged result should be loose files")
+    zen    = unrealcmd.Opt(False, "Staged result should use data streaming from zenserver")
 
 
 
@@ -160,4 +169,5 @@ class Deploy(_Impl):
         self.args.cook   = False
         self.args.deploy = True
         self.args.nopak  = True
+        self.args.zen    = False
         return super().main(skipstage=True)
diff --git a/Engine/Extras/ushell/channels/unreal/core/cmds/zen.py b/Engine/Extras/ushell/channels/unreal/core/cmds/zen.py
--- a/Engine/Extras/ushell/channels/unreal/core/cmds/zen.py
+++ b/Engine/Extras/ushell/channels/unreal/core/cmds/zen.py
@@ -10,6 +10,44 @@ import unreal.zencmd
 import http.client
 import json
 import hashlib
+import tempfile
+
+#-------------------------------------------------------------------------------
+def _get_ips_socket():
+    import socket
+    def _get_ip(family, ip_addr):
+        with socket.socket(family, socket.SOCK_DGRAM) as s:
+            try:
+                s.connect((ip_addr, 1))
+                return s.getsockname()[0]
+            except Exception as e:
+                return
+    ret = (_get_ip(socket.AF_INET6, "fe80::0"), _get_ip(socket.AF_INET, "172.31.255.255"))
+    return [x for x in ret if x]
+
+def _get_ips_wmic():
+    import subprocess as sp
+    ret = []
+    current_ips = None
+    wmic_cmd = ("wmic", "nicconfig", "where", "ipenabled=true", "get", "ipaddress,ipconnectionmetric", "/format:value")
+    with sp.Popen(wmic_cmd, stdout=sp.PIPE, stderr=sp.DEVNULL) as proc:
+        for line in proc.stdout:
+            if line.startswith(b"IPAddress="):
+                if line := line[10:].strip():
+                    current_ips = line
+            elif line.startswith(b"IPConnectionMetric="):
+                if line := line[19:].strip():
+                    connection_metric = int(line)
+                    for ip in eval(current_ips):
+                        ret.append((connection_metric, ip))
+    ret.sort(key=lambda p: p[0])
+    ret = [item[1] for item in ret]
+    ret.sort(key=lambda x: "." in x) # sort ipv4 last
+    return ret
+
+_get_ips = _get_ips_wmic if os.name == "nt" else _get_ips_socket
+
+
 
 #-------------------------------------------------------------------------------
 class Dashboard(unreal.zencmd.ZenUETargetBaseCmd):
@@ -24,10 +62,12 @@ class Start(unreal.zencmd.ZenUETargetBaseCmd):
     SponsorProcessID = unrealcmd.Opt("", "Process ID to be added as a sponsor process for the zenserver process")
 
     def main(self):
-
         args = []
         if self.args.SponsorProcessID:
             args.append("-SponsorProcessID=" + str(self.args.SponsorProcessID))
+        elif sys.platform != 'win32':
+            grandparent_pid = os.popen("ps -p %d -oppid=" % os.getppid()).read().strip()
+            args.append("-SponsorProcessID=" + grandparent_pid)
         else:
             args.append("-SponsorProcessID=" + str(os.getppid()))
 
@@ -61,8 +101,10 @@ class ImportSnapshot(unreal.zencmd.ZenUtilityBaseCmd):
     snapshotindex = unrealcmd.Arg(0, "0-based index of the snapshot within the snapshot descriptor to import from")
     projectid = unrealcmd.Opt("", "The zen project ID to import into (defaults to an ID based on the current ushell project)")
     oplog = unrealcmd.Opt("", "The zen oplog to import into (defaults to the oplog name in the snapshot)")
+    sourcehost = unrealcmd.Opt("", "The source host to import from (defaults to the host specified in the snapshot)")
     cloudauthservice = unrealcmd.Opt("", "Name of the service to authorize with when importing from a cloud source")
     asyncimport = unrealcmd.Opt(False, "Trigger import but don't wait for completion")
+    forceimport = unrealcmd.Opt(False, "Force import of all attachments")
 
     def complete_projectid(self, prefix):
         return self.perform_project_completion(prefix)
@@ -102,8 +144,12 @@ class ImportSnapshot(unreal.zencmd.ZenUtilityBaseCmd):
         if sys.platform == "darwin": bin_type = "osx-x64"
         elif sys.platform == "linux": bin_type = "linux-x64"
 
+        bin_ext = ".exe"
+        if sys.platform == "darwin": bin_ext = ""
+        elif sys.platform == "linux": bin_ext = ""
+
         bin_path = engine.get_dir()
-        bin_path /= f"Binaries/DotNET/OidcToken/{bin_type}/OidcToken.exe"
+        bin_path /= f"Binaries/DotNET/OidcToken/{bin_type}/OidcToken{bin_ext}"
         if not bin_path.is_file():
             raise FileNotFoundError(f"Unable to find '{bin_path}'")
 
@@ -112,9 +158,12 @@ class ImportSnapshot(unreal.zencmd.ZenUtilityBaseCmd):
         if not self._service_name:
             raise ValueError("Unable to discover service name")
 
+        token_dir = tempfile.TemporaryDirectory()
+        token_file_path = f"{token_dir.name}/oidctoken.json"
         oidcargs = (
             "--ResultToConsole=true",
             "--Service=" + str(self._service_name),
+            f"--OutFile={token_file_path}"
         )
 
         if project := ue_context.get_project():
@@ -125,7 +174,9 @@ class ImportSnapshot(unreal.zencmd.ZenUtilityBaseCmd):
         if ret != 0:
             return False
 
-        tokenresponse = json.loads(output)
+        tokenresponse = None
+        with open(token_file_path, "r") as token_file:
+            tokenresponse = json.load(token_file)
 
         self._AccessToken = tokenresponse['Token']
         self.print_info("Cloud access token obtained for import operation.")
@@ -135,54 +186,149 @@ class ImportSnapshot(unreal.zencmd.ZenUtilityBaseCmd):
     def get_exec_context(self):
         context = super().get_exec_context()
         if hasattr(self, '_AccessToken') and self._AccessToken:
-            context.get_env()["UE-CloudDataCacheAccessToken"] = self._AccessToken
+            if sys.platform == 'win32':
+                context.get_env()["UE-CloudDataCacheAccessToken"] = self._AccessToken
+            else:
+                context.get_env()["UE_CloudDataCacheAccessToken"] = self._AccessToken
         return context
 
+    def conform_to_cook_forms(self, oplogid):
+        ue_context = self.get_unreal_context()
+        platforms = ue_context.get_platform_provider()
+        for platform_name in platforms.read_platform_names():
+            platform = platforms.get_platform(platform_name)
+            for form_name in ("client", "game", "server"):
+                try:
+                    cook_form = platform.get_cook_form(form_name)
+                    if cook_form.lower() == oplogid.lower():
+                        return cook_form
+                except:
+                    pass
+        return oplogid
+
+    def conform_hostname(self, hostname):
+        return hostname
+
     def add_args_from_descriptor(self, snapshot, args):
         snapshot_type = snapshot['type']
 
-        # TODO: file, zen snapshot types
+        sourcehost = None
+        if snapshot_type == 'cloud' or snapshot_type == 'zen':
+            sourcehost = snapshot['host']
+            if self.args.sourcehost:
+                sourcehost = self.conform_hostname(self.args.sourcehost)
+
         if snapshot_type == 'cloud':
             args.append('--cloud')
-            args.append(snapshot['host'])
+            args.append(sourcehost)
             args.append('--namespace')
             args.append(snapshot['namespace'])
             args.append('--bucket')
             args.append(snapshot['bucket'])
             args.append('--key')
             args.append(snapshot['key'])
+        elif snapshot_type == 'zen':
+            args.append('--zen')
+            args.append(sourcehost)
+            args.append('--source-project')
+            args.append(snapshot['projectid'])
+            args.append('--source-oplog')
+            args.append(snapshot['oplogid'])
+        elif snapshot_type == 'file':
+            args.append('--file')
+            args.append(snapshot['directory'])
+            args.append('--name')
+            args.append(snapshot['filename'])
         else:
             self.print_error(f"Unsupported snapshot type {snapshot_type}")
             return False
 
         return True
 
-    @unrealcmd.Cmd.summarise
-    def main(self):
-        try:
-            with open(self.args.snapshotdescriptor, "rt") as file:
-                descriptor = json.load(file)
-        except FileNotFoundError:
-            self.print_error(f"Error accessing snapshot descriptor file {self.args.snapshotdescriptor}")
-            return False
+    def perform_target_creation(self, snapshot):
+        target_projectid = self.get_project_or_default(self.args.projectid)
+        if self.find_project_by_id(target_projectid) is None:
+            ue_context = self.get_unreal_context()
+            args = [
+            "project-create",
+            target_projectid,
+            ue_context.get_branch().get_dir(),
+            ue_context.get_engine().get_dir(),
+            ue_context.get_project().get_dir(),
+            ue_context.get_project().get_path()
+            ]
+            self.run_zen_utility(args)
+
+        target_oplogid = self.args.oplog or snapshot['targetplatform']
 
-        snapshot = descriptor["snapshots"][self.args.snapshotindex]
-        snapshot_type = snapshot['type']
+        ue_context = self.get_unreal_context()
+        target_gcmarker = ue_context.get_project().get_dir() / f"Saved/Cooked/{self.conform_to_cook_forms(target_oplogid)}/ue.projectstore"
+
+        target_gcmarker_dir = os.path.dirname(target_gcmarker)
+        if not os.path.exists(target_gcmarker_dir):
+            os.makedirs(target_gcmarker_dir)
+        projectstore = {
+            "zenserver" : {
+                "islocalhost" : True,
+                "hostname" : "[::1]",
+                "remotehostnames" : _get_ips(),
+                "hostport" : 8558,
+                "projectid" : target_projectid,
+                "oplogid" : target_oplogid
+            }
+        }
+        with open(target_gcmarker, "w") as target_gcmarker_outfile:
+            json.dump(projectstore, target_gcmarker_outfile, indent='\t')
 
+        args = [
+        "oplog-create",
+        target_projectid,
+        target_oplogid,
+        target_gcmarker,
+        "--force-update"
+        ]
+        self.run_zen_utility(args)
+        return True
+
+    def perform_import(self, snapshot):
         args = [
         "oplog-import",
         self.get_project_or_default(self.args.projectid),
         self.args.oplog or snapshot['targetplatform'],
+        "--ignore-missing-attachments",
+        "--clean"
         ]
 
         if self.args.asyncimport:
             args.append("--async")
 
+        if self.args.forceimport:
+            args.append("--force")
+
         if not self.add_args_from_descriptor(snapshot, args):
             return False
 
-        if snapshot_type == 'cloud':
+        if snapshot['type'] == 'cloud':
             if not self.refresh_zen_token():
                 return False
 
         return self.run_zen_utility(args)
+
+
+    @unrealcmd.Cmd.summarise
+    def main(self):
+        self.launch_zenserver_if_not_running()
+        try:
+            with open(self.args.snapshotdescriptor, "rt") as file:
+                descriptor = json.load(file)
+        except FileNotFoundError:
+            self.print_error(f"Error accessing snapshot descriptor file {self.args.snapshotdescriptor}")
+            return False
+
+        snapshot = descriptor["snapshots"][self.args.snapshotindex]
+
+        if not self.perform_target_creation(snapshot):
+            self.print_error(f"Error creating import target location")
+            return False
+
+        return self.perform_import(snapshot)
diff --git a/Engine/Extras/ushell/channels/unreal/core/describe.flow.py b/Engine/Extras/ushell/channels/unreal/core/describe.flow.py
--- a/Engine/Extras/ushell/channels/unreal/core/describe.flow.py
+++ b/Engine/Extras/ushell/channels/unreal/core/describe.flow.py
@@ -180,6 +180,10 @@ sln_open = flow.describe.Command()
 sln_open.source("cmds/sln.py", "Open")
 sln_open.invoke("sln", "open")
 
+sln_10x = flow.describe.Command()
+sln_10x.source("cmds/sln.py", "Open10x")
+sln_10x.invoke("sln", "open", "10x")
+
 #-------------------------------------------------------------------------------
 uat = flow.describe.Command()
 uat.source("cmds/uat.py", "Uat")
@@ -220,6 +224,19 @@ project_change = flow.describe.Command()
 project_change.source("cmds/project.py", "Change")
 project_change.invoke("project")
 
+#-------------------------------------------------------------------------------
+odsc_client = flow.describe.Command()
+odsc_client.source("cmds/odsc.py", "Client")
+odsc_client.invoke("cook", "odsc", "client")
+
+odsc_game = flow.describe.Command()
+odsc_game.source("cmds/odsc.py", "Game")
+odsc_game.invoke("cook", "odsc", "game")
+
+odsc_all = flow.describe.Command()
+odsc_all.source("cmds/odsc.py", "All")
+odsc_all.invoke("cook", "odsc", "all")
+
 #-------------------------------------------------------------------------------
 prompt = flow.describe.Command()
 prompt.source("prompt.py", "Prompt")
diff --git a/Engine/Extras/ushell/channels/unreal/core/pylib/unreal/platforms/win64/android.py b/Engine/Extras/ushell/channels/unreal/core/pylib/unreal/platforms/win64/android.py
--- a/Engine/Extras/ushell/channels/unreal/core/pylib/unreal/platforms/win64/android.py
+++ b/Engine/Extras/ushell/channels/unreal/core/pylib/unreal/platforms/win64/android.py
@@ -109,6 +109,10 @@ class Platform(unreal.Platform):
     def _get_package_name(self):
         config = self.get_unreal_context().get_config()
         if ret := str(config.get("Engine", Platform.config_section, "PackageName")):
+            if "[PROJECT]" in ret:
+                project_name = self.get_unreal_context().get_project().get_name()
+                project_name = project_name.replace("-", "_")
+                ret = ret.replace("[PROJECT]", project_name)
             return ret
         raise ValueError(f"Failed querying '{Platform.config_section}/PackageName'")
 
@@ -121,8 +125,17 @@ class Platform(unreal.Platform):
         print("Intent: ", intent)
         print()
 
-        args = "'" + " ".join(x.replace("\"", "\\\"") for x in args) + "'"
         adb = self._get_adb()
+        reverse_port_mappings = [
+                    41899,  # Network file server, DEFAULT_TCP_FILE_SERVING_PORT in NetworkMessage.h
+                    1981,   # Unreal Insights data collection, TraceInsightsModule.cpp
+                    8558,   # Zen Store
+        ]
+        for reverse_port_mapping in reverse_port_mappings:
+            reverse_cmd = (str(adb), "reverse", "tcp:"+str(reverse_port_mapping), "tcp:"+str(reverse_port_mapping), ">", "NUL")
+            os.system(" ".join(reverse_cmd))
+
+        args = "'" + " ".join(x.replace("\"", "\\\"") for x in args) + "'"
         cmd = (str(adb), "-d", "shell", "am", "start", "--es", "cmdline", args, intent)
         print(*cmd)
 
diff --git a/Engine/Extras/ushell/channels/unreal/core/pylib/unreal/zencmd.py b/Engine/Extras/ushell/channels/unreal/core/pylib/unreal/zencmd.py
--- a/Engine/Extras/ushell/channels/unreal/core/pylib/unreal/zencmd.py
+++ b/Engine/Extras/ushell/channels/unreal/core/pylib/unreal/zencmd.py
@@ -9,9 +9,107 @@
 import http.client
 import json
 import hashlib
+import socket
 
 #-------------------------------------------------------------------------------
 class ZenBaseCmd(unrealcmd.Cmd):
+
+    def _build(self, target, variant, platform):
+        if target.get_type() == unreal.TargetType.EDITOR:
+            build_cmd = ("_build", "editor", variant)
+        else:
+            platform = platform or unreal.Platform.get_host()
+            build_cmd = ("_build", "target", target.get_name(), platform, variant)
+
+        import subprocess
+        ret = subprocess.run(build_cmd)
+        if ret.returncode:
+            raise SystemExit(ret.returncode)
+
+    def get_binary_path(self, target, variant, platform, build):
+        ue_context = self.get_unreal_context()
+
+        if isinstance(target, str):
+            target = ue_context.get_target_by_name(target)
+        else:
+            target = ue_context.get_target_by_type(target)
+
+        if build:
+            self._build(target, variant, platform)
+
+        variant = unreal.Variant.parse(variant)
+        if build := target.get_build(variant, platform):
+            if build := build.get_binary_path():
+                return build
+
+        raise EnvironmentError(f"No {variant} build found for target '{target.get_name()}'")
+
+    def run_ue_program_target(self, target, variant, build, cmdargs, runargs):
+        binary_path = self.get_binary_path(target, variant, None, build)
+
+        final_args = []
+        if cmdargs: final_args += cmdargs
+        if runargs: final_args += (*unreal.cmdline.read_ueified(*runargs),)
+
+        launch_kwargs = {}
+        if not sys.stdout.isatty():
+            launch_kwargs = { "silent" : True }
+        exec_context = self.get_exec_context()
+        runnable = exec_context.create_runnable(str(binary_path), *final_args);
+        runnable.launch(**launch_kwargs)
+        return True if runnable.is_gui() else runnable.wait()
+
+    def launch_zenserver(self):
+        args = ["-SponsorProcessID=" + str(os.getppid())]
+        self.run_ue_program_target("ZenLaunch", "development", True, args, None)
+
+    def launch_zenserver_if_not_running(self):
+        if not self.is_zenserver_running():
+            self.launch_zenserver()
+
+    def is_zenserver_running(self, hostname="localhost", port=8558):
+        headers = {"Accept":"application/json"}
+        conn = http.client.HTTPConnection(hostname, port=port, timeout=0.1)
+        try:
+            conn.request("GET", "/health/ready", headers=headers)
+            response = conn.getresponse()
+            return response.read().decode().lower() == "ok!"
+        except OSError:
+            return False
+
+        return False
+
+    def find_project_by_id(self, projectid, hostname="localhost", port=8558):
+        self.launch_zenserver_if_not_running()
+        headers = {"Accept":"application/json"}
+        conn = http.client.HTTPConnection(hostname, port=port, timeout=0.1)
+        projects = []
+        try:
+            conn.request("GET", "/prj", headers=headers)
+            response = conn.getresponse()
+            projects = json.loads(response.read().decode())
+        except OSError:
+            return None
+
+        for project in projects:
+            if project['Id'] == projectid:
+                return project
+        return None
+
+    def find_oplog_by_id(self, projectid, oplogid, hostname="localhost", port=8558):
+        self.launch_zenserver_if_not_running()
+        headers = {"Accept":"application/json"}
+        conn = http.client.HTTPConnection(hostname, port=port, timeout=0.1)
+        oplogs = []
+        try:
+            conn.request("GET", f"/prj/{projectid}/oplog/{oplogid}", headers=headers)
+            response = conn.getresponse()
+            oplog = json.loads(response.read().decode())
+        except:
+            return None
+
+        return oplog
+
     def perform_project_completion(self, prefix, hostname="localhost", port=8558):
         headers = {"Accept":"application/json"}
         conn = http.client.HTTPConnection(hostname, port=port, timeout=0.1)
@@ -20,7 +118,8 @@
             conn.request("GET", "/prj", headers=headers)
             response = conn.getresponse()
             projects = json.loads(response.read().decode())
-        except socket.timeout:
+        except OSError:
+            self.launch_zenserver()
             return
 
         for project in projects:
@@ -39,7 +138,8 @@
                 conn.request("GET", "/prj", headers=headers)
                 response = conn.getresponse()
                 projects = json.loads(response.read().decode())
-        except socket.timeout:
+        except OSError:
+            self.launch_zenserver()
             return
 
         for project in projects:
@@ -69,50 +169,8 @@
     runargs  = unrealcmd.Arg([str], "Additional arguments to pass to the process being launched")
     build    = unrealcmd.Opt(True, "Build the target prior to running it")
 
-    def _build(self, target, variant, platform):
-        if target.get_type() == unreal.TargetType.EDITOR:
-            build_cmd = ("_build", "editor", variant)
-        else:
-            platform = platform or unreal.Platform.get_host()
-            build_cmd = ("_build", "target", target.get_name(), platform, variant)
-
-        import subprocess
-        ret = subprocess.run(build_cmd)
-        if ret.returncode:
-            raise SystemExit(ret.returncode)
-
-    def get_binary_path(self, target, platform=None):
-        ue_context = self.get_unreal_context()
-
-        if isinstance(target, str):
-            target = ue_context.get_target_by_name(target)
-        else:
-            target = ue_context.get_target_by_type(target)
-
-        if self.args.build:
-            self._build(target, self.args.variant, platform)
-
-        variant = unreal.Variant.parse(self.args.variant)
-        if build := target.get_build(variant, platform):
-            if build := build.get_binary_path():
-                return build
-
-        raise EnvironmentError(f"No {self.args.variant} build found for target '{target.get_name()}'")
-
-    def run_ue_program_target(self, target_name, cmdargs = None):
-        binary_path = self.get_binary_path(target_name)
-
-        final_args = []
-        if cmdargs: final_args += cmdargs
-        if self.args.runargs: final_args += (*unreal.cmdline.read_ueified(*self.args.runargs),)
-
-        launch_kwargs = {}
-        if not sys.stdout.isatty():
-            launch_kwargs = { "silent" : True }
-        exec_context = self.get_exec_context()
-        runnable = exec_context.create_runnable(str(binary_path), *final_args);
-        runnable.launch(**launch_kwargs)
-        return True if runnable.is_gui() else runnable.wait()
+    def run_ue_program_target(self, target, cmdargs = None):
+        return super().run_ue_program_target(target, self.args.variant, self.args.build, cmdargs, self.args.runargs)
 
 
 #-------------------------------------------------------------------------------
diff --git a/Engine/Extras/ushell/channels/unreal/perforce/cmds/authors.py b/Engine/Extras/ushell/channels/unreal/perforce/cmds/authors.py
new file mode 100644
--- /dev/null
+++ b/Engine/Extras/ushell/channels/unreal/perforce/cmds/authors.py
@@ -0,0 +1,319 @@
+# Copyright Epic Games, Inc. All Rights Reserved.
+
+import datetime
+import flow.cmd
+import math
+import multiprocessing.dummy
+import os
+import re
+import subprocess
+import sys
+import time
+
+p4_annotate_walltime = 0
+p4_annotate_call_count = 0
+p4_describe_walltime = 0
+p4_describe_call_count = 0
+p4_where_walltime = 0
+p4_where_call_count = 0
+
+def query_cl_linecount(file_path, print_error, print_warning, follow_integrations=False):
+    start_time = time.time()
+
+    follow_flag = "-i"
+    if follow_integrations:
+        follow_flag = "-I"
+
+    # Run p4 annotate to get the CLs for each line.
+    annotate_result = subprocess.run(["p4", "annotate", follow_flag, file_path], capture_output=True)
+    annotate_stdout = annotate_result.stdout.decode("utf-8")
+    annotate_stderr = annotate_result.stderr.decode("utf-8")
+    if annotate_result.returncode != 0:
+        msg = "p4 annotate failed, giving up. Output:\n%s\n%s" % (annotate_stdout, annotate_stderr)
+        print_error(msg)
+        sys.exit(1)
+
+    global p4_annotate_walltime
+    p4_annotate_walltime += time.time() - start_time
+    global p4_annotate_call_count
+    p4_annotate_call_count += 1
+
+    # Give the user feedback if the annotated path cannot be found.
+    if "no such file(s)" in annotate_stderr:
+        print_error("p4 annotate call failed, giving up. Output:\n%s\n%s" % (annotate_stdout, annotate_stderr))
+        sys.exit(1)
+
+    annotate_lines = annotate_stdout.splitlines()
+
+    # We count the number of lines changed by each CL.
+    cl_linecount_map = {}
+
+    # p4 annotate has a header followed by a output that looks like
+    # "CL_NUMBER: LINE_CONTENT". Let's grab that CL_NUMBER with a regex!
+    cl_regex = re.compile(r'^([0-9]+):')
+
+    for line in annotate_lines:
+        m = cl_regex.match(line)
+        if not m:
+            continue
+        matching_string = m.group(1)
+        changelist_number = int(matching_string)
+
+        if changelist_number in cl_linecount_map:
+            cl_linecount_map[changelist_number] += 1
+        else:
+            cl_linecount_map[changelist_number] = 1
+
+    return cl_linecount_map
+
+def query_owner_scores(cl_linecount_map, print_info, print_warning, print_error, start_date=None):
+    owner_score_map = {}
+
+    cl_count = len(cl_linecount_map)
+
+    cl_process_list = []
+
+    # Use at most 100 parallel calls to Perforce.
+    pool = multiprocessing.dummy.Pool(100)
+
+    def p4_describe(cl):
+        describe_result = subprocess.run(["p4", "describe", "-s", str(cl)], capture_output=True)
+        try:
+            stdout_string = describe_result.stdout.decode("utf-8")
+            return cl, stdout_string.splitlines()
+        except UnicodeDecodeError:
+            pass
+        try:
+            stdout_string = describe_result.stdout.decode("cp1252")
+            return cl, stdout_string.splitlines()
+        except UnicodeDecodeError:
+            pass
+
+        print_warning("Failed to decode output of \"p4 describe -s %d\". Results might be incorrect. Raw output was:" % cl, describe_result.stdout)
+
+        return cl, []
+
+    start_time = time.time()
+
+    cl_lines_map = {}
+    for cl, lines in pool.map(p4_describe, cl_linecount_map.keys()):
+        cl_lines_map[cl] = lines
+        global p4_describe_call_count
+        p4_describe_call_count += 1
+
+    global p4_describe_walltime
+    p4_describe_walltime += time.time() - start_time
+
+    # p4 describe has a header followed by the list of affected files formatted like
+    # ... //P4/path/to/file.extension#REVISION action. Let's count those lines!
+    changed_lines_regex = re.compile(r'^\.\.\.')
+
+    # p4 describe mentiones the author as "Change 123 by Person.Name@Branch on date".
+    # Let's grab the author's name from that!
+    describe_owner_regex = re.compile(r'.* by (\S+)@\S+ on .*')
+
+    # p4 describe prints the date of a changelist. Grab it!
+    date_regex = re.compile(r'.* on (\d\d\d\d/\d\d/\d\d) .*')
+
+    start_date_skip_count = 0
+
+    for cl, linecount in cl_linecount_map.items():
+        describe_lines = cl_lines_map[cl]
+
+        owner = None
+        date_found = False
+        skip_cl = False
+        match_count = 0
+
+        for line in describe_lines:
+            if start_date and not date_found and not skip_cl:
+                date_match = date_regex.match(line)
+                if date_match:
+                    date_string = date_match.group(1)
+                    date = datetime.datetime.strptime(date_string, "%Y/%m/%d")
+                    date_found = True
+
+                    # Skip this CL if it's dated before our start date.
+                    if date < start_date:
+                        start_date_skip_count += 1
+                        skip_cl = True
+
+            if changed_lines_regex.match(line):
+                match_count += 1
+
+            if not owner:
+                owner_match = describe_owner_regex.match(line)
+                if owner_match:
+                    owner = owner_match.group(1).lower()
+
+        if skip_cl:
+            owner = "<OUTSIDE TIME SPAN>"
+
+        # Idea 1: The more files a CL changes, the less important that CL is to each
+        #         file. Therefore, the author of the CL gains less ownership of each
+        #         file than an author targeting fewer files.
+        # Idea 2: The more lines a CL changes in a file, the more important the CL
+        #         is to the file and therefore the more ownership the author gets.
+        if match_count == 0:
+            score = 0
+        else:
+            # TODO: This score should only divide by the number of files that
+            # are not in the set of files we're querying ownership for. A CL
+            # that affects a lot of files but where all are among the files we
+            # are querying for should not be penalized.
+            score = linecount * 1.0 / math.sqrt(match_count)
+
+        assert(owner != "")
+        assert(owner is not None)
+
+        if owner in owner_score_map:
+            owner_score_map[owner] += score
+        else:
+            owner_score_map[owner] = score
+
+    if start_date_skip_count > 0:
+        assert(start_date is not None)
+        print_warning("Skipped %d of %d CLs because of the provided start date %s." % (start_date_skip_count, cl_count, start_date.strftime("%Y/%m/%d")))
+
+        if start_date_skip_count == cl_count:
+            print_warning("Skipped all CLs. Nothing to do. Exiting.")
+
+    return owner_score_map
+
+def sort_ownership_score(owner_score_map):
+    return sorted(owner_score_map.items(), key=lambda item: item[1], reverse=True)
+
+def combine_ownership_score(left_owner_map, right_owner_map):
+    out_map = {}
+
+    for in_map in [left_owner_map, right_owner_map]:
+        for key in in_map:
+            if key in out_map:
+                out_map[key] += in_map[key]
+            else:
+                out_map[key] = in_map[key]
+
+    return out_map
+
+def normalize_scores(scores):
+    total_score = 0.0
+    for owner, score in scores:
+        total_score += score
+
+    if total_score == 0:
+        return scores
+
+    normalized_scores = []
+    for owner, score in scores:
+        new_score = 100.0 * score / total_score
+        normalized_scores.append((owner, new_score))
+
+    return normalized_scores
+
+def print_scores(scores, normalized):
+    if normalized:
+        scores = normalize_scores(scores)
+
+        for owner,score in scores:
+            print("%5.2f%% %s" % (score, owner))
+    else:
+        for owner,score in scores:
+            print("%.2f %s" % (score, owner))
+
+def expand_directory(path, print_info, print_error):
+    """Strip trailing slashes and instead determine using "p4 where" combined
+    with os.path.isdir if a path is a directory or not. Paths that do point at
+    directories get ... added to the end (delimited by the system path
+    separator)."""
+
+    # Nothing to do if the user already indicated clearly that they intend to
+    # query a directory.
+    if path.endswith("\\...") or path.endswith("/..."):
+        return path
+
+    output_path = path
+
+    while output_path.endswith("\\") or output_path.endswith("/"):
+        output_path = output_path.rstrip("\\").rstrip("/")
+
+    start_time = time.time()
+
+    where_result = subprocess.run(["p4", "-z", "tag", "where", output_path], capture_output=True)
+    where_stdout = where_result.stdout.decode("utf-8")
+
+    global p4_where_walltime
+    p4_where_walltime += time.time() - start_time
+    global p4_where_call_count
+    p4_where_call_count += 1
+
+    if where_result.returncode != 0:
+        where_stderr = where_result.stderr.decode("utf-8")
+        msg = "p4 where failed, giving up. Output:\n%s\n%s" % (where_stdout, where_stderr)
+        print_error(msg)
+        sys.exit(1)
+
+    local_path = None
+    path_regex = re.compile(r'^... path (.*)')
+    for line in where_stdout.splitlines():
+        m = path_regex.match(line)
+        if not m:
+            continue
+        local_path = m.group(1)
+
+    if os.path.isdir(local_path):
+        output_path = os.path.join(output_path, "...")
+
+    if output_path != path:
+        print_info("Adjusted the provided path. Will use \"%s\"" % output_path)
+
+    return output_path
+
+
+class Authors(flow.cmd.Cmd):
+    """List the authors of a given path.
+    Give each author a score based on how many lines they have changed under the
+    given path. An author's score is increased if they changed many lines (more
+    important author) and decreased if many files where involved in their
+    changelists (less important author because the file was not as central to
+    their change)."""
+    path                = flow.cmd.Arg(str, "Path to the file. Supports relative and absolute local paths. Supports Perforce paths (including those ending in ...).")
+    after               = flow.cmd.Opt("", "Only count changes after a specified date (following the format YYYY/mm/dd).")
+    pastyear            = flow.cmd.Opt(False, "Only count changes from the past year.")
+    pasttwoyears        = flow.cmd.Opt(False, "Only count changes from the past two years.")
+    time                = flow.cmd.Opt(False, "Print performance timings.")
+    follow_integrations = flow.cmd.Opt(False, "Follow integrations when annotating file lines by passing -I to p4 annotate. The default is to follow branches by passing -i to p4 annotate.")
+    normalize_scores    = flow.cmd.Opt(True, "Normalize scores before printing them (default).")
+
+    def main(self):
+        start_time = time.time()
+
+        start_date = None
+        if len(self.args.after) > 0:
+            try:
+                start_date = datetime.datetime.strptime(self.args.after, "%Y/%m/%d")
+            except:
+                self.print_error("Failed to parse your provided date \"%s\" as a YYYY/mm/dd date. Exiting." % self.args.after)
+                sys.exit(1)
+        elif self.args.pastyear:
+            start_date = datetime.datetime.now() + datetime.timedelta(days=-365)
+        elif self.args.pasttwoyears:
+            start_date = datetime.datetime.now() + datetime.timedelta(days=-2*365)
+
+        if start_date is not None:
+            self.print_info("Only counting changes since %s." % start_date.strftime("%Y/%m/%d"))
+
+        path = expand_directory(self.args.path, self.print_info, self.print_error)
+
+        cl_linecount_map = query_cl_linecount(path, self.print_error, self.print_warning, follow_integrations=self.args.follow_integrations)
+        owner_score = query_owner_scores(cl_linecount_map, self.print_info, self.print_warning, self.print_error, start_date)
+        scores = sort_ownership_score(owner_score)
+
+        print_scores(scores, self.args.normalize_scores)
+
+        run_time = time.time() - start_time
+
+        if self.args.time:
+            self.print_info("Total runtime: %5.2f s" % run_time)
+            self.print_info("  p4 annotate: %5.2f s (%d %s)" % (p4_annotate_walltime, p4_annotate_call_count, "call" if p4_annotate_call_count == 1 else "calls"))
+            self.print_info("  p4 describe: %5.2f s (%d %s)" % (p4_describe_walltime, p4_describe_call_count, "call" if p4_describe_call_count == 1 else "calls"))
+            self.print_info("     p4 where: %5.2f s (%d %s)" % (p4_where_walltime, p4_where_call_count, "call" if p4_where_call_count == 1 else "calls"))
diff --git a/Engine/Extras/ushell/channels/unreal/perforce/cmds/workspace.py b/Engine/Extras/ushell/channels/unreal/perforce/cmds/workspace.py
--- a/Engine/Extras/ushell/channels/unreal/perforce/cmds/workspace.py
+++ b/Engine/Extras/ushell/channels/unreal/perforce/cmds/workspace.py
@@ -51,6 +51,7 @@ def _populate_ue_branch_engine(depot_path, dry_run):
     minimal_sync_specs = (
         depot_path / "Engine/Build/Build.version",
         depot_path / "Engine/Source/*.Target.cs",
+        depot_path / "Engine/Config/BaseEngine.ini",
         depot_path / "*",
     )
     if dry_run:
diff --git a/Engine/Extras/ushell/channels/unreal/perforce/describe.flow.py b/Engine/Extras/ushell/channels/unreal/perforce/describe.flow.py
--- a/Engine/Extras/ushell/channels/unreal/perforce/describe.flow.py
+++ b/Engine/Extras/ushell/channels/unreal/perforce/describe.flow.py
@@ -71,6 +71,11 @@ who = flow.describe.Command()
 who.source("cmds/who.py", "Who")
 who.invoke("p4", "who")
 
+#-------------------------------------------------------------------------------
+authors = flow.describe.Command()
+authors.source("cmds/authors.py", "Authors")
+authors.invoke("p4", "authors")
+
 #-------------------------------------------------------------------------------
 boot = flow.describe.Command()
 boot.source("boot.py", "Boot")
diff --git a/Engine/Extras/ushell/ushell.sh b/Engine/Extras/ushell/ushell.sh
--- a/Engine/Extras/ushell/ushell.sh
+++ b/Engine/Extras/ushell/ushell.sh
@@ -1,16 +1,16 @@
 #!/bin/bash
 
-# Check whether this script has been "sourced" into a host shell
+# Launch as a sub-shell if the the script was not sourced into the current shell
 if [ $0 = "$BASH_SOURCE" ]; then
-    # The script was not sourced, so source it in a child shell and propagate
-    # any arguments
     args=""
     for arg in "$@"
     do
         escaped=$(printf '%q' "$arg")
         args="$args $escaped"
     done
-    bash --init-file <(echo "source \"$HOME/.bashrc\"; source \"$BASH_SOURCE\" $args")
+
+    bash_rc="$HOME/.bashrc"
+    bash --init-file <(echo "if [ -f '$bash_rc' ]; then source '$bash_rc'; fi; source '$BASH_SOURCE' $args")
     exit $?
 fi
 
