Updated ushell [35016981]

#rnx

diff --git a/Engine/Extras/ushell/channels/unreal/core/pylib/unreal/_context.py b/Engine/Extras/ushell/channels/unreal/core/pylib/unreal/_context.py
--- a/Engine/Extras/ushell/channels/unreal/core/pylib/unreal/_context.py
+++ b/Engine/Extras/ushell/channels/unreal/core/pylib/unreal/_context.py
@@ -227,7 +227,8 @@ class _Target(object):
             path += "-" + platform + "-" + variant.ueify()
         path += ".target"
 
-        if receipt := next(self._ufs_mount.glob(path), None):
+        ufs_node = self._ufs_mount.get_node()
+        if receipt := next(ufs_node.glob(path), None):
             return _Build(receipt, variant, platform, self._ufs_mount)
 
 
diff --git a/Engine/Extras/ushell/channels/unreal/perforce/cmds/gui.py b/Engine/Extras/ushell/channels/unreal/perforce/cmds/gui.py
--- a/Engine/Extras/ushell/channels/unreal/perforce/cmds/gui.py
+++ b/Engine/Extras/ushell/channels/unreal/perforce/cmds/gui.py
@@ -60,7 +60,6 @@ class Gui(flow.cmd.Cmd):
             "-p", p4_port,
             "-u", username,
             "-c", client,
-            "-t", "pending",
         )
 
         if self.args.filename.is_file():
