Updated ushell [35052661]

#rnx

diff --git a/Engine/Extras/ushell/channels/flow/core/cmds/debug.py b/Engine/Extras/ushell/channels/flow/core/cmds/debug.py
--- a/Engine/Extras/ushell/channels/flow/core/cmds/debug.py
+++ b/Engine/Extras/ushell/channels/flow/core/cmds/debug.py
@@ -7,6 +7,7 @@ import pprint
 import hashlib
 import marshal
 import flow.cmd
+from pathlib import Path
 from urllib.request import urlopen, URLError
 
 #-------------------------------------------------------------------------------
@@ -104,11 +105,19 @@ class Paths(flow.cmd.Cmd):
     """ Prints behind-the-scenes paths. """
 
     def main(self):
+        for item in Path(__file__).parents:
+            if (item / "ushell.bat").is_file():
+                deploy_dir = item
+                break
+        else:
+            deploy_path = "*unknown*"
+
         channel = self.get_channel()
         system = channel.get_system()
-        print("State:", system.get_working_dir())
-        print("Tools:", system.get_tools_dir())
-        print(" Temp:", system.get_temp_dir())
+        print("Deploy:", deploy_dir)
+        print(" State:", system.get_working_dir())
+        print(" Tools:", system.get_tools_dir())
+        print("  Temp:", system.get_temp_dir())
 
 #-------------------------------------------------------------------------------
 class _Channels(flow.cmd.Cmd):
diff --git a/Engine/Extras/ushell/channels/unreal/core/cmds/build.py b/Engine/Extras/ushell/channels/unreal/core/cmds/build.py
--- a/Engine/Extras/ushell/channels/unreal/core/cmds/build.py
+++ b/Engine/Extras/ushell/channels/unreal/core/cmds/build.py
@@ -55,16 +55,25 @@ class _Builder(object):
     def read_actions(self):
         platform = self._platform
         variant = self._variant
-        targets = (x.get_name() for x in self._targets)
 
         args = self._args
+        if len(self._targets) > 1:
+            targets = []
+            for target in self._targets:
+                target = f"-Target={target.get_name()} {platform} {variant}"
+                targets.append(target)
+                variant = "development" # \o/
+        else:
+            targets = (x.get_name() for x in self._targets)
+            args = (platform, variant, *args)
+
         projected = any(x.is_project_target() for x in self._targets)
         if self._projected or projected:
             if project := self._ue_context.get_project():
                 args = ("-Project=" + str(project.get_path()), *args)
 
         ubt = self._ue_context.get_engine().get_ubt()
-        yield from ubt.read_actions(*targets, platform, variant, *args)
+        yield from ubt.read_actions(*targets, *args)
 
 
 
@@ -444,9 +453,7 @@ class Editor(_BuildCmd, unrealcmd.MultiPlatformCmd):
 
         editor_only = self.is_constrained_build() # i.e. single file or module
         editor_only |= bool(self.args.analyze)
-        editor_only |= self.args.variant != "development"
         if not editor_only:
-
             targets_to_add = []
             if not self.args.noscw:
                 targets_to_add.append("ShaderCompileWorker")
