Updated ushell [35363050]

#rnx

diff --git a/Engine/Extras/ushell/README.txt b/Engine/Extras/ushell/README.txt
--- a/Engine/Extras/ushell/README.txt
+++ b/Engine/Extras/ushell/README.txt
@@ -79,6 +79,26 @@ prompt to enter a search string and will display the latest command with a match
 Further Ctrl-R hits will step backwards through commands that match the search
 string (with Ctrl-S stepping forwards). History searching is case-sensitive.
 
+# Integrating with UnrealGameSync
+
+UGS supports additional tools deployable via Perforce via the UGS setting
+`DeploymentSettings/ToolsDepotPath`. By adding a `Tools/ushell/ushell.ini` file
+users can enable a ushell status panel link via `Options > Application
+Settings`. An example `ushell.ini` follows;
+
+```
+[Settings]
+Id=922EED87-E732-464C-92DC-5A8F7ED955E2
+Name=ushell
+Description=ushell
+SafeWhenBusy=1
++StatusPanelLinks=(Label="ushell", FileName="$(COMSPEC)", Arguments="/c \"\"$(BranchDir)\\Engine\\Extras\\ushell\\ushell.bat\" --project=\"$(ProjectFile)\"\"", WorkingDir="$(ProjectDir)")
+```
+
+It is also possible to distribute a copy of ushell as a `ushell.zip` file
+alongside `ushell.ini`. This can be useful to deploy ushell to users regardless
+of which branch they may be working on.
+
 # Scripting
 
 There is modest support for scripting ushell with Batch scripts. This can be
diff --git a/Engine/Extras/ushell/channels/flow/core/prompt.py b/Engine/Extras/ushell/channels/flow/core/prompt.py
--- a/Engine/Extras/ushell/channels/flow/core/prompt.py
+++ b/Engine/Extras/ushell/channels/flow/core/prompt.py
@@ -3,6 +3,7 @@
 import os
 import flow.cmd
 
+#-------------------------------------------------------------------------------
 class Prompt(flow.cmd.Cmd):
     """ Used to collect key/value pairs for tags in the prompt. """
     format = flow.cmd.Opt("none", "Output format. Options; none,sh")
diff --git a/Engine/Extras/ushell/channels/flow/nt/describe.flow.py b/Engine/Extras/ushell/channels/flow/nt/describe.flow.py
--- a/Engine/Extras/ushell/channels/flow/nt/describe.flow.py
+++ b/Engine/Extras/ushell/channels/flow/nt/describe.flow.py
@@ -4,9 +4,9 @@ import flow.describe
 
 #-------------------------------------------------------------------------------
 clink = flow.describe.Tool()
-clink.version("1.0.0a4")
+clink.version("1.0.0a5")
 clink.payload("[REDACTED_PROFILE_URL]/clink/releases/download/$VERSION/clink-$VERSION.zip")
-clink.sha1("6a642bc48e21a3093698067f0fe09eff3ccdfaa5")
+clink.sha1("70289e92e3313a2b0e8dee801901eae61f8992a3")
 clink.platform("win32")
 clink.bin("clink_x64.exe")
 
diff --git a/Engine/Extras/ushell/channels/flow/nt/provision.bat b/Engine/Extras/ushell/channels/flow/nt/provision.bat
--- a/Engine/Extras/ushell/channels/flow/nt/provision.bat
+++ b/Engine/Extras/ushell/channels/flow/nt/provision.bat
@@ -7,8 +7,8 @@ goto:eof
 ::------------------------------------------------------------------------------
 :main
 
-set _pyver=3.12.2
-set _pysha=144ee3fa0c4171f98ec554ac01bb44ee1e895d154609f143b1d20b52c3fe383e
+set _pyver=3.12.4
+set _pysha=15fea3c9367653a85086fe37216b4d1a1c78688fa5e1587e1db0b0f658856564
 set _pytag=312
 set _pymark=%_pyver%.version
 
diff --git a/Engine/Extras/ushell/channels/unreal/core/cmds/zen.py b/Engine/Extras/ushell/channels/unreal/core/cmds/zen.py
--- a/Engine/Extras/ushell/channels/unreal/core/cmds/zen.py
+++ b/Engine/Extras/ushell/channels/unreal/core/cmds/zen.py
@@ -314,6 +314,19 @@ class ImportSnapshot(unreal.zencmd.ZenUtilityBaseCmd):
 
         return self.run_zen_utility(args)
 
+    def perform_clean_post_import(self, snapshot):
+        target_oplogid = self.args.oplog or snapshot['targetplatform']
+
+        ue_context = self.get_unreal_context()
+        target_dir = ue_context.get_project().get_dir() / f"Saved/Cooked/{self.conform_to_cook_forms(target_oplogid)}"
+
+        for fsentry in os.scandir(target_dir):
+            if fsentry.is_dir(follow_symlinks=False):
+                shutil.rmtree(fsentry.path)
+            else:
+                if fsentry.name != 'ue.projectstore':
+                    os.unlink(fsentry.path)
+
 
     @unrealcmd.Cmd.summarise
     def main(self):
@@ -331,4 +344,8 @@ class ImportSnapshot(unreal.zencmd.ZenUtilityBaseCmd):
             self.print_error(f"Error creating import target location")
             return False
 
-        return self.perform_import(snapshot)
+        import_result = self.perform_import(snapshot)
+        if import_result == 0:
+            self.perform_clean_post_import(snapshot)
+        return import_result
+
