Updated ushell [35547518]

diff --git a/Engine/Extras/ushell/.p4ignore.txt b/Engine/Extras/ushell/.p4ignore.txt
new file mode 100644
--- /dev/null
+++ b/Engine/Extras/ushell/.p4ignore.txt
@@ -0,0 +1,2 @@
+# Ignore Python cached files
+**/*.pyc
diff --git a/Engine/Extras/ushell/channels/flow/core/cmds/debug.py b/Engine/Extras/ushell/channels/flow/core/cmds/debug.py
--- a/Engine/Extras/ushell/channels/flow/core/cmds/debug.py
+++ b/Engine/Extras/ushell/channels/flow/core/cmds/debug.py
@@ -55,7 +55,7 @@ class Debug(flow.cmd.Cmd):
 
         print()
         session = self.get_noticeboard(self.Noticeboard.SESSION)
-        pprint.pprint(session._get_inner())
+        session.debug(pprint.pprint)
 
 #-------------------------------------------------------------------------------
 class Invalidate(flow.cmd.Cmd):
diff --git a/Engine/Extras/ushell/channels/flow/core/cmds/help.py b/Engine/Extras/ushell/channels/flow/core/cmds/help.py
--- a/Engine/Extras/ushell/channels/flow/core/cmds/help.py
+++ b/Engine/Extras/ushell/channels/flow/core/cmds/help.py
@@ -49,8 +49,11 @@ class Help(flow.cmd.Cmd):
         for name, cmd_class in output:
             dots = lead_dot[len(name):]
             leader = flow.cmd.text.light_cyan(name) + " " + flow.cmd.text.grey(dots)
-            desc = cmd_class.get_desc(pretty=desc_width)
-            desc = desc.replace("\n", "\n" + lead_spc)
+            if isinstance(cmd_class, str):
+                desc = cmd_class
+            else:
+                desc = cmd_class.get_desc(pretty=desc_width)
+                desc = desc.replace("\n", "\n" + lead_spc)
             print(" ", leader, sep="", end="")
             print(desc)
             print()
diff --git a/Engine/Extras/ushell/channels/flow/core/system/flow/_noticeboard.py b/Engine/Extras/ushell/channels/flow/core/system/flow/_noticeboard.py
new file mode 100644
--- /dev/null
+++ b/Engine/Extras/ushell/channels/flow/core/system/flow/_noticeboard.py
@@ -0,0 +1,131 @@
+import os
+import marshal
+from typing import Any
+from pathlib import Path
+
+#-------------------------------------------------------------------------------
+class _Journaler(object):
+    _dir        : Path
+    _name       : str
+    _uniqueness : int   = os.getpid()
+
+    def __init__(self, path:Path) -> None:
+        self._dir = path.parent
+        self._name = path.name
+
+    @classmethod
+    def _get_uniqueness(self) -> int:
+        _Journaler._uniqueness += 1
+        return _Journaler._uniqueness
+
+    def _get_serial_range(self) -> tuple[int, int]:
+        min_index = 1 << 31
+        max_index = -1
+        for item in self._dir.glob(self._name + ".*"):
+            try:
+                candidate = int(item.suffix[1:])
+                min_index = min(min_index, candidate)
+                max_index = max(max_index, candidate)
+            except ValueError:
+                pass
+        return min(min_index, max_index), max_index
+
+    def read(self) -> dict:
+        _, serial = self._get_serial_range()
+        if serial < 0:
+            return {}
+
+        try:
+            path = self._dir / (self._name + f".{serial}")
+            for i in range(5):
+                try:
+                    with path.open("rb") as inp:
+                        if data := marshal.load(inp):
+                            return data
+                    return {}
+                except PermissionError as e:
+                    continue
+        except OSError:
+            pass
+
+        return {}
+
+    def flush(self, edits:dict) -> None:
+        self._dir.mkdir(parents=True, exist_ok=True)
+        make_path = lambda x: self._dir / (self._name + "." + str(x))
+        for retry in range(50):
+            min_serial, serial = self._get_serial_range()
+            if serial >= 0:
+                try:
+                    with make_path(serial).open("rb") as inp:
+                        data = marshal.load(inp)
+                except IOError as e:
+                    import time
+                    time.sleep(0.02)
+                    continue
+            else:
+                data = {}
+
+            data.update(edits)
+
+            # Write edited board to a temporary file
+            temp_name = f"wr_{_Journaler._get_uniqueness()}"
+            temp_path = make_path(temp_name)
+            try:
+                with temp_path.open("wb") as out:
+                    marshal.dump(data, out)
+            except IOError as e:
+                return
+
+            # Try and swap edited temp file into place
+            serial += 1
+            dest_path = make_path(serial)
+            try:
+                temp_path.rename(dest_path)
+                break
+            except (FileExistsError, PermissionError):
+                temp_path.unlink()
+                continue
+            except Exception as e:
+                temp_path.unlink()
+                continue
+
+        for index in range(min_serial, serial - 8):
+            try: make_path(index).unlink(missing_ok=True)
+            except: break
+
+#-------------------------------------------------------------------------------
+class Noticeboard(object):
+    _edits          : dict
+    _context        : int
+    _journaler      : _Journaler
+
+    def __init__(self, board_path:Path|str, context:int|str=-1) -> None:
+        self._edits = {}
+        self._context = context
+        self._journaler = _Journaler(Path(board_path))
+
+    def __del__(self) -> None:
+        self.flush()
+
+    def __setitem__(self, key:str, value:Any) -> None:
+        sub_edit = self._edits.setdefault(self._context, {})
+        sub_edit[key] = value
+
+    def __getitem__(self, name:str) -> Any:
+        data = self._journaler.read()
+        if sub_data := data.get(self._context):
+            return sub_data.get(name)
+
+    def get(self, name:str, default:Any) -> Any:
+        ret = self[name]
+        return default if ret is None else ret
+
+    def flush(self) -> None:
+        if self._edits:
+            self._journaler.flush(self._edits)
+            self._edits = {}
+
+    def debug(self, print_fn) -> None:
+        data = self._journaler.read()
+        print_fn(data)
diff --git a/Engine/Extras/ushell/channels/flow/core/system/flow/cmd.py b/Engine/Extras/ushell/channels/flow/core/system/flow/cmd.py
--- a/Engine/Extras/ushell/channels/flow/core/system/flow/cmd.py
+++ b/Engine/Extras/ushell/channels/flow/core/system/flow/cmd.py
@@ -6,6 +6,7 @@ import marshal
 from enum import Enum
 
 from . import _backtrace
+from ._noticeboard import Noticeboard
 from .text import *
 
 #-------------------------------------------------------------------------------
@@ -133,40 +134,6 @@ class _Channel(object):
 
 
 
-#-------------------------------------------------------------------------------
-class _Noticeboard(object):
-    def __init__(self, path):
-        self._path = path
-        self._inner = None
-        self._modified = False
-
-    def __del__(self):
-        if self._modified:
-            with open(self._path, "wb") as x:
-                marshal.dump(self._inner, x)
-
-    def _get_inner(self):
-        if not self._inner:
-            try:
-                with open(self._path, "rb") as x:
-                    self._inner = marshal.load(x)
-            except:
-                self._inner = {}
-        return self._inner
-
-    def __getitem__(self, name):
-        return self._get_inner().get(name)
-
-    def __setitem__(self, name, value):
-        self._modified = True
-        self._get_inner()[name] = value
-
-    def copy_from(self, other):
-        self._modified = True
-        self._inner = other._get_inner().copy()
-
-
-
 #-------------------------------------------------------------------------------
 from . import _flick
 Arg = _flick.Arg
@@ -222,11 +189,11 @@ class Cmd(_flick.Cmd, _Channel):
 
     def get_noticeboard(self, board_type):
         temp_dir = self._channel.get_system().get_temp_dir()
+        flow_sid = -1
         if board_type == Cmd.Noticeboard.SESSION:
             flow_sid = os.getenv("FLOW_SID", "x")
-            return _Noticeboard(temp_dir + "session_" + flow_sid)
-        elif board_type == Cmd.Noticeboard.PERSISTENT:
-            return _Noticeboard(temp_dir + "noticeboard_p")
+            flow_sid = int(flow_sid) if flow_sid.isdecimal() else "-493"
+        return Noticeboard(temp_dir + "noticeboard", int(flow_sid))
 
     def edit_file(self, path):
         if editor := os.getenv("GIT_EDITOR") or os.getenv("P4EDITOR"):
diff --git a/Engine/Extras/ushell/channels/flow/core/system/flow/system.py b/Engine/Extras/ushell/channels/flow/core/system/flow/system.py
--- a/Engine/Extras/ushell/channels/flow/core/system/flow/system.py
+++ b/Engine/Extras/ushell/channels/flow/core/system/flow/system.py
@@ -229,7 +229,7 @@ class System(object):
         return self._working_dir
 
     def get_temp_dir(self):
-        ret = self.get_working_dir() + "temp/"
+        ret = self.get_working_dir() + "../temp/"
         try: os.mkdir(ret)
         except FileExistsError: pass
         return ret
diff --git a/Engine/Extras/ushell/channels/unreal/core/cmds/project.py b/Engine/Extras/ushell/channels/unreal/core/cmds/project.py
--- a/Engine/Extras/ushell/channels/unreal/core/cmds/project.py
+++ b/Engine/Extras/ushell/channels/unreal/core/cmds/project.py
@@ -246,7 +246,7 @@ class Change(flow.cmd.Cmd):
         # Work out what the candidate should be
         if candidate.lower() == "cwd":
             return self._set_cwd()
-        elif candidate.lower() == "none" or candidate.lower() == "branch":
+        elif candidate.lower() in ("none", "branch"):
             candidate = self._set_branch()
         elif "/" in candidate or "\\" in candidate or os.path.exists(candidate):
             candidate = self._set_qualified(candidate)
diff --git a/Engine/Extras/ushell/channels/unreal/core/cmds/sln.py b/Engine/Extras/ushell/channels/unreal/core/cmds/sln.py
--- a/Engine/Extras/ushell/channels/unreal/core/cmds/sln.py
+++ b/Engine/Extras/ushell/channels/unreal/core/cmds/sln.py
@@ -3,6 +3,7 @@
 import os
 import shutil
 import unrealcmd
+from pathlib import Path
 
 #-------------------------------------------------------------------------------
 class _Base(unrealcmd.MultiPlatformCmd):
@@ -110,7 +111,7 @@ class Generate(_Base):
 
 #-------------------------------------------------------------------------------
 class Open(_Base):
-    """ Opens project files in Visual Studio """
+    """ Opens project files in Visual Studio. """
 
     def get_exec_context(self):
         context = super().get_exec_context()
@@ -163,3 +164,94 @@ class Open10x(_Base):
         result = cmd.launch()
         return True
 
+
+
+#-------------------------------------------------------------------------------
+class Tiny(_Base):
+    """ Quickly generates and opens a tiny solution. Suitable for those who use VS
+    just for debugging and an alternative file opener (e.g the fzf-based rmenu). A
+    run of '.sln generate' is not required. """
+
+    def _write_primary_name(self):
+        if hasattr(self, "_primary_name"):
+            return self._primary_name
+
+        ue_context = self.get_unreal_context()
+        if (branch := ue_context.get_branch()) is None:
+            return
+
+        engine = ue_context.get_engine()
+        name = "UE" + str(engine.get_version_major())
+        name = name + "_" + ue_context.get_branch().get_name()
+        ppn_path = engine.get_dir() / "Intermediate/ProjectFiles/PrimaryProjectName.txt"
+        ppn_path.parent.mkdir(parents=True, exist_ok=True)
+        with ppn_path.open("wt") as out:
+            print(name, file=out)
+
+    def main(self):
+        ue_context = self.get_unreal_context()
+        branch = ue_context.get_branch()
+        if not branch:
+            raise EnvironmentError("A branch ius required for a tiny solution")
+
+        self._write_primary_name()
+
+        projects = []
+        if branch := ue_context.get_branch():
+            for item in branch.read_projects():
+                projects += item.parent.glob("Source/*.Target.cs")
+
+        programs = []
+        for item in ue_context.glob("Source/Programs/*"):
+            programs += item.glob("*.Target.cs")
+
+        import slnformer
+
+        proj_dir = ue_context.get_engine().get_dir()
+        proj_dir /= "Intermediate/ProjectFiles/TinySln"
+
+        former = slnformer.Sln(self._get_sln_path(), proj_dir)
+
+        tiny_proj = former.add_project("_tiny", makefile=False)
+        tiny_proj.add_property("ushell_bin_suffix", "-Win64-$(Configuration)", cond="$(Configuration) != 'Development'")
+        tiny_proj.add_property("ushell_bin_prefix", "../../../")
+        tiny_proj.add_property("ushell_bin", "$(ushell_bin_prefix)Binaries/Win64/$(ushell_target)$(ushell_bin_suffix).exe")
+        tiny_proj.add_property("NMakeOutput", "$(ushell_bin)")
+        tiny_proj.set_build_action("cd $(ushell_dir) &amp; .build target $(ushell_target) win64 $(Configuration.ToLower())")
+        tiny_proj.set_rebuild_action("cd $(ushell_dir) &amp; $(NMakeCleanCommandLine) &amp; $(NMakeBuildCommandLine)")
+        tiny_proj.set_clean_action("cd $(ushell_dir) &amp; $(NMakeBuildCommandLine) --clean")
+        tiny_proj.add_property("FLOW_SID", "")
+        tiny_proj.add_import("$(VCTargetsPath)/Microsoft.Cpp.Default.props")
+        tiny_proj.add_import("$(VCTargetsPath)/Microsoft.Cpp.props")
+        tiny_proj.add_import("$(VCTargetsPath)/Microsoft.Cpp.targets")
+
+        def prep_proj(proj, context_dir):
+            target = proj.get_name().replace(".vcxproj", "")
+            proj.add_import(tiny_proj.get_name())
+            proj.add_property("ushell_target", target)
+            proj.add_property("ushell_dir", str(context_dir))
+
+        for project in projects:
+            proj_dir = project.parent.parent
+            proj_name = project.name[:-10]
+            folder = former.add_folder(proj_dir.name)
+            proj = folder.add_project(proj_name)
+            prep_proj(proj, proj_dir)
+
+        folder = former.add_folder("_programs")
+        for program in programs:
+            prog_dir = project.parent
+            prog_name = program.name[:-10]
+            prog_index = folder.add_folder(prog_name[0].lower())
+            proj = prog_index.add_project(prog_name)
+            prep_proj(proj, prog_dir)
+
+        former.write()
+
+        self.print_info("Creating a tiny solution");
+        print("Sln path:", self._get_sln_path())
+        print("Dir:", proj_dir)
+        print("Projects:", len(projects))
+        print("Programs:", len(programs))
+
+        return self._open_sln()
diff --git a/Engine/Extras/ushell/channels/unreal/core/cmds/stage.py b/Engine/Extras/ushell/channels/unreal/core/cmds/stage.py
--- a/Engine/Extras/ushell/channels/unreal/core/cmds/stage.py
+++ b/Engine/Extras/ushell/channels/unreal/core/cmds/stage.py
@@ -9,10 +9,13 @@ class _Impl(unrealcmd.Cmd):
     target   = unrealcmd.Arg(str, "Cooked target to stage")
     platform = unrealcmd.Arg("", "Platform whose cook should be staged")
     variant  = unrealcmd.Arg("development", "Build variant to stage binaries to copy")
+    style    = unrealcmd.Arg("auto", "Type of result to stage ('auto' makes the best choice for common use)")
     uatargs  = unrealcmd.Arg([str], "Additional arguments to pass to UAT")
 
     complete_target = ("game", "client", "server")
 
+    complete_style = ("pak", "nopak", "zen","auto")
+
     def _build(self, build_editor, ubt_args=None):
         self.print_info("Building")
 
@@ -109,13 +112,23 @@ class _Impl(unrealcmd.Cmd):
 
         platform = self.get_platform(platform)
 
-        if self.args.zen:
-            cook_form = platform.get_cook_form(target_type.name.lower())
-            projectstore_file = project.get_dir() / f"Saved/Cooked/{cook_form}/" / "ue.projectstore"
-            if not projectstore_file.is_file():
-                self.print_error(f"'{projectstore_file}' missing, cannot use --zen for '{self.args.target}' target")
+        cook_form = platform.get_cook_form(target_type.name.lower())
+        projectstore_file = project.get_dir() / f"Saved/Cooked/{cook_form}/" / "ue.projectstore"
+        cooked_to_zen = projectstore_file.is_file()
+
+        uat_packaging_args = "-pak"
+
+        if cooked_to_zen and not (self.args.style == 'pak'):
+            uat_packaging_args = None
+
+        if self.args.style == 'zen':
+            if not cooked_to_zen:
+                self.print_error(f"'{projectstore_file}' missing, cannot use zen output style for '{self.args.target}' target")
                 return 1
-            self.args.nopak = True
+            uat_packaging_args = None
+
+        if self.args.style == 'nopak':
+            uat_packaging_args = None
 
         cook_flavor = platform.get_cook_flavor()
 
@@ -132,7 +145,7 @@ class _Impl(unrealcmd.Cmd):
             "-config=" + self.args.variant,
             "-platform=" + platform.get_name(),
             ("-cookflavor=" + cook_flavor) if cook_flavor else None,
-            None if self.args.nopak else "-pak",
+            uat_packaging_args,
             *self.args.uatargs,
         )
 
@@ -155,8 +168,6 @@ class Stage(_Impl):
     build  = unrealcmd.Opt((False, ""), "Build code prior to staging")
     cook   = unrealcmd.Opt((False, ""), "Run a cook before running the stage step")
     deploy = unrealcmd.Opt(False, "Deploy to devkit/device after staging")
-    nopak  = unrealcmd.Opt(False, "Staged result should be loose files")
-    zen    = unrealcmd.Opt(False, "Staged result should use data streaming from zenserver")
 
 
 
diff --git a/Engine/Extras/ushell/channels/unreal/core/cmds/uat.py b/Engine/Extras/ushell/channels/unreal/core/cmds/uat.py
--- a/Engine/Extras/ushell/channels/unreal/core/cmds/uat.py
+++ b/Engine/Extras/ushell/channels/unreal/core/cmds/uat.py
@@ -7,11 +7,11 @@ import unrealcmd
 class Uat(unrealcmd.MultiPlatformCmd):
     """ Runs an UnrealAutomationTool command """
     command     = unrealcmd.Arg(str, "The UnrealAutomationTool command to run")
-    unprojected = unrealcmd.Opt(False, "No implicit '-project=' please")
     uatargs     = unrealcmd.Arg([str], "Arguments to pass to the UAT command")
+    unprojected = unrealcmd.Opt(False, "No implicit '-project=' please")
+    allscripts  = unrealcmd.Opt(False, "Ask UAT to compile all scripts (default is project-only)");
 
     def complete_command(self, prefix):
-        import os
         import re
 
         def _read_foreach_subdir(dir, func):
@@ -61,9 +61,11 @@ class Uat(unrealcmd.MultiPlatformCmd):
         cmd /= "RunUAT.bat" if os.name == "nt" else "RunUAT.sh"
         args = (self.args.command,)
 
-        if args[0] and not self.args.unprojected:
-            if project := ue_context.get_project():
+        if project := ue_context.get_project():
+            if args[0] and not self.args.unprojected:
                 args = (*args, "-project=" + str(project.get_path()))
+            if not self.args.allscripts:
+                args = (*args, "-ScriptsForProject=" + project.get_name())
 
         exec_context = self.get_exec_context()
         cmd = exec_context.create_runnable(cmd, *args, *self.args.uatargs)
diff --git a/Engine/Extras/ushell/channels/unreal/core/describe.flow.py b/Engine/Extras/ushell/channels/unreal/core/describe.flow.py
--- a/Engine/Extras/ushell/channels/unreal/core/describe.flow.py
+++ b/Engine/Extras/ushell/channels/unreal/core/describe.flow.py
@@ -184,6 +184,10 @@ sln_10x = flow.describe.Command()
 sln_10x.source("cmds/sln.py", "Open10x")
 sln_10x.invoke("sln", "open", "10x")
 
+sln_tiny = flow.describe.Command()
+sln_tiny.source("cmds/sln.py", "Tiny")
+sln_tiny.invoke("sln", "open", "tiny")
+
 #-------------------------------------------------------------------------------
 uat = flow.describe.Command()
 uat.source("cmds/uat.py", "Uat")
diff --git a/Engine/Extras/ushell/channels/unreal/core/pylib/slnformer.py b/Engine/Extras/ushell/channels/unreal/core/pylib/slnformer.py
new file mode 100644
--- /dev/null
+++ b/Engine/Extras/ushell/channels/unreal/core/pylib/slnformer.py
@@ -0,0 +1,243 @@
+from typing import Iterable, Optional, IO
+from pathlib import Path
+
+#-------------------------------------------------------------------------------
+class _XmlNode(object):
+    def __init__(self, parent:Optional["_XmlNode"]=None) -> None:
+        self._parent = parent
+        self._depth = parent._depth + 1 if parent else -1
+
+    def get_out(self) -> IO:
+        if self._parent:
+            return self._parent.get_out()
+
+    def __getattr__(self, tag:str) -> None:
+        assert tag[0] != "_"
+        def thunk(text:str|None=None, **attrs) -> "_XmlTag":
+            return _XmlTag(self, tag, text, **attrs)
+        return thunk
+
+#-------------------------------------------------------------------------------
+class _XmlTag(_XmlNode):
+    def __init__(self, parent:_XmlNode, tag:str, text:str|None, **attrs) -> None:
+        super().__init__(parent)
+        self._tag = tag
+        self._text = text
+        self._write(" " * self._depth, "<", tag, end="")
+        for k,v in attrs.items():
+            self._write(" ", k, rf'="{v}"', end="")
+
+    def __del__(self):
+        if self._tag:
+            suffix = "/>" if self._text is None else f">{self._text}</{self._tag}>"
+            self._write(suffix)
+
+    def _write(self, *args, **kwargs):
+        print(*args, **kwargs, sep="", file=self.get_out())
+
+    def __enter__(self):
+        self._write(">")
+        return self
+
+    def __exit__(self, *exc):
+        self._write(" " * self._depth, f"</{self._tag}>")
+        self._tag = None
+
+#-------------------------------------------------------------------------------
+class _Xmler(_XmlNode):
+    def __init__(self, out:IO) -> None:
+        super().__init__()
+        print(r'<?xml version="1.0" encoding="utf-8"?>', file=out)
+        self._out = out
+
+    def get_out(self) -> IO:
+        return self._out
+
+
+
+#-------------------------------------------------------------------------------
+class _Node(object):
+    guid_count = 0
+
+    def __init__(self, name:str, parent:"_Node") -> None:
+        self._parent = parent
+        self._name = name
+        self._guid = f"{{0A9E0A9E-0A9E-0A9E-0A9E-0A9E0A9E{_Node.guid_count:04d}}}"
+        _Node.guid_count += 1
+
+    def get_name(self) -> str:       return self._name
+    def get_guid(self) -> str:       return self._guid
+    def get_parent(self) -> "_Node": return None if self._parent == self else self._parent
+
+#-------------------------------------------------------------------------------
+class _Project(_Node):
+    def __init__(self, name:str, parent:"_Folder") -> None:
+        super().__init__(name, parent)
+        self._props = []
+        self._imports = { "pre" : [], "post" : [] }
+
+    def set_build_action(self, cmd:str, *, cond:str="") -> None:
+        return self.add_property("NMakeBuildCommandLine", cmd, cond=cond)
+
+    def set_rebuild_action(self, cmd:str, *, cond:str="") -> None:
+        return self.add_property("NMakeReBuildCommandLine", cmd, cond=cond)
+
+    def set_clean_action(self, cmd:str, *, cond:str="") -> None:
+        return self.add_property("NMakeCleanCommandLine", cmd, cond=cond)
+
+    def add_property(self, name:str, value:str, *, cond:str="") -> None:
+        self._props.append((name, value, cond))
+
+    def add_import(self, imp_path:Path, *, pre:bool=False) -> None:
+        self._imports["pre" if pre else "post"].append(str(imp_path))
+
+    def write(self, out:IO) -> None:
+        cond_props = {}
+        for k,v,c in self._props:
+            prop = (k, v)
+            cond_props.setdefault(c, []).append(prop)
+
+        xmler = _Xmler(out)
+        with xmler.Project(DefaultTargets="Build", xmlns="http://schemas.microsoft.com/developer/msbuild/2003") as node:
+            # TODO: formalise this!
+            node._write( " <ItemGroup>")
+            node._write(r'  <ProjectConfiguration Include="Debug|x64"><Configuration>Debug</Configuration><Platform>x64</Platform></ProjectConfiguration>')
+            node._write(r'  <ProjectConfiguration Include="Development|x64"><Configuration>Development</Configuration><Platform>x64</Platform></ProjectConfiguration>')
+            node._write( " </ItemGroup>")
+
+            for imp in self._imports["pre"]:
+                node.Import(Project=imp)
+
+            for cond, props in cond_props.items():
+                kwargs = {"Condition" : cond} if cond else {}
+                with node.PropertyGroup(**kwargs) as prop_grp:
+                    for k, v in props:
+                        getattr(prop_grp, k)(v)
+
+            for imp in self._imports["post"]:
+                node.Import(Project=imp)
+
+#-------------------------------------------------------------------------------
+class _Folder(_Node):
+    def __init__(self, name:str, parent:_Node) -> None:
+        self._children = []
+        self._folders = {}
+        super().__init__(name, parent)
+
+    def add_folder(self, name:str) -> "_Folder":
+        if name in self._folders:
+            return self._folders[name]
+        ret = _Folder(name, self)
+        self._children.append(ret)
+        self._folders[name] = ret
+        return ret
+
+    def add_project(self, name:str, makefile:bool=True) -> _Project:
+        ext = ".vcxproj" if makefile else ".proj"
+        ret = _Project(name + ext, self)
+        if makefile:
+            ret.add_property("ConfigurationType", "Makefile")
+            ret.add_property("ProjectGuid", ret.get_guid())
+        self._children.append(ret)
+        return ret
+
+    def read_children(self) -> Iterable[_Node]:
+        for child in self._children:
+            yield child
+            if isinstance(child, _Folder):
+                yield from child.read_children()
+
+    def read_folders(self) -> Iterable["_Folder"]:
+        return (x for x in self.read_children() if isinstance(x, _Folder))
+
+    def read_projects(self) -> Iterable[_Project]:
+        return (x for x in self.read_children() if isinstance(x, _Project))
+
+#-------------------------------------------------------------------------------
+class Sln(_Folder):
+    def __init__(self, sln_path:Path, proj_dir:Path) -> None:
+        super().__init__("", self)
+        self._sln_path = sln_path
+        self._proj_dir = proj_dir
+
+    def write(self):
+        for item in (self._sln_path.parent, self._proj_dir):
+            item.mkdir(parents=True, exist_ok=True)
+
+        for project in self.read_projects():
+            proj_path = self._proj_dir / project.get_name()
+
+            with proj_path.open("wt") as out:
+                project.write(out)
+
+        with self._sln_path.open("wt") as out:
+            write = lambda *x: print(*x, sep="", file=out)
+
+            sln_header = (
+                "Microsoft Visual Studio Solution File, Format Version 12.00",
+                "# Visual Studio Version 17",
+                "VisualStudioVersion = 17.0.0.0",
+                "MinimumVisualStudioVersion = 10.0.0.0",
+            )
+            sln_files = (
+                # TODO: formalise this!
+                r'Project("{2150E333-8FDC-42A3-9474-1A3956D46DE8}") = "__sln_files", "__sln_files", "{51012800-7858-4A8D-BEF1-3878E3C2D44B}"',
+                "\tProjectSection(SolutionItems) = preProject",
+                "\t\tEngine/Extras/VisualStudioDebugging/Unreal.natvis = Engine/Extras/VisualStudioDebugging/Unreal.natvis",
+                "\tEndProjectSection",
+                "EndProject",
+            )
+            sln_project = (
+                r'Project("{{8BC9CEB8-8B4A-11D0-8D11-00A0C91BC942}}") = "{stem}", "{path}", "{guid}"',
+                # "\tProjectSection(ProjectDependencies) = postProject",
+                # "\t\t{guid} = {guid}",
+                # "\tEndProjectSection",
+                "EndProject",
+            )
+            sln_folder = (
+                r'Project("{{2150E333-8FDC-42A3-9474-1A3956D46DE8}}") = "{name}", "{name}", "{guid}"',
+                "EndProject",
+            )
+            sln_nest = (
+                "\tGlobalSection(NestedProjects) = preSolution",
+                # {folder_guid} = {child_guid}
+                "\tEndGlobalSection",
+            )
+            sln_global = (
+                "Global",
+                "EndGlobal",
+            )
+
+            for line in sln_header: write(line)
+            for line in sln_files:  write(line)
+
+            for folder in self.read_folders():
+                guid = folder.get_guid()
+                name = folder.get_name()
+                for line in sln_folder:
+                    write(line.format(guid=guid, name=name))
+
+            proj_dir = self._proj_dir.relative_to(self._sln_path.parent)
+            for project in self.read_projects():
+                name = project.get_name()
+                if ".vcxproj" not in name:
+                    continue
+
+                fmt_args = {
+                    "stem" : Path(name).stem,
+                    "guid" : project.get_guid(),
+                    "path" : str(proj_dir / name),
+                }
+                for line in sln_project:
+                    write(line.format(**fmt_args))
+
+            write(sln_global[0])
+
+            write(sln_nest[0])
+            for item in self.read_children():
+                parent = item.get_parent()
+                if isinstance(parent, _Folder):
+                    write("\t\t", item.get_guid(), " = ", parent.get_guid());
+            write(sln_nest[1])
+
+            write(sln_global[1])
