Updated ushell [35621077]

#rnx

diff --git a/Engine/Extras/ushell/channels/flow/core/cmds/debug.py b/Engine/Extras/ushell/channels/flow/core/cmds/debug.py
--- a/Engine/Extras/ushell/channels/flow/core/cmds/debug.py
+++ b/Engine/Extras/ushell/channels/flow/core/cmds/debug.py
@@ -169,7 +169,7 @@ class Sha1s(_Channels):
             self._impl(name, manifest)
 
 #-------------------------------------------------------------------------------
-class Argumentss(flow.cmd.Cmd):
+class Arguments(flow.cmd.Cmd):
     """ Prints the given arguments as they would be received by commands """
     arguments = flow.cmd.Arg([str], "Arguments to print")
 
diff --git a/Engine/Extras/ushell/channels/flow/core/describe.flow.py b/Engine/Extras/ushell/channels/flow/core/describe.flow.py
--- a/Engine/Extras/ushell/channels/flow/core/describe.flow.py
+++ b/Engine/Extras/ushell/channels/flow/core/describe.flow.py
@@ -35,7 +35,7 @@ debug_sha1s.invoke("debug", "tools", "sha1s")
 debug_sha1s.prefix("$")
 
 debug_args = flow.describe.Command()
-debug_args.source("cmds/debug.py", "Argumentss")
+debug_args.source("cmds/debug.py", "Arguments")
 debug_args.invoke("debug", "args")
 debug_args.prefix("$")
 
diff --git a/Engine/Extras/ushell/channels/flow/core/system/flow/_flick.py b/Engine/Extras/ushell/channels/flow/core/system/flow/_flick.py
--- a/Engine/Extras/ushell/channels/flow/core/system/flow/_flick.py
+++ b/Engine/Extras/ushell/channels/flow/core/system/flow/_flick.py
@@ -102,7 +102,8 @@ def Opt(type_value, *args):
 
 #-------------------------------------------------------------------------------
 class _ArgOptBox(object):
-    def __init__(self, items):
+    def __init__(self, parent, items):
+        super().__setattr__("_parent", parent)
         super().__setattr__("_items", items)
 
     def __iter__(self):
@@ -118,6 +119,11 @@ class _ArgOptBox(object):
     def is_default(self, name):
         return self._items[name][1]
 
+    def get_type(self, name):
+        arg_info = getattr(self._parent, name, None)
+        if isinstance(arg_info, _ArgOptBase):
+            return arg_info.get_type()
+
 
 
 #-------------------------------------------------------------------------------
@@ -322,7 +328,7 @@ class Cmd(object):
 
         return "\n".join(paras[:-1])
 
-    def invoke(self, in_args):
+    def invoke(self, in_args, *, invoke_path=""):
         self.validate()
 
         def read_in_args():
@@ -356,7 +362,8 @@ class Cmd(object):
             out_type = argopt.get_type().__name__
             on_error(f"Unable to convert '{value}' to type '{out_type}'")
 
-        self.args = _ArgOptBox(args_out)
+        self.args = _ArgOptBox(self, args_out)
+        self._invoke_path = invoke_path
 
         return self._call_main()
 
@@ -390,7 +397,7 @@ class Cmd(object):
             return
 
         if hasattr(completer, "__call__"):
-            self.args = _ArgOptBox(args_out)
+            self.args = _ArgOptBox(self, args_out)
             completer = completer(prefix)
 
         if completer:
diff --git a/Engine/Extras/ushell/channels/flow/core/system/flow/cmd.py b/Engine/Extras/ushell/channels/flow/core/system/flow/cmd.py
--- a/Engine/Extras/ushell/channels/flow/core/system/flow/cmd.py
+++ b/Engine/Extras/ushell/channels/flow/core/system/flow/cmd.py
@@ -113,24 +113,46 @@ class _ExecContext(object):
 
 
 #-------------------------------------------------------------------------------
-class _Channel(object):
-    def read_extensions(self, mount):
-        extendable = self.get_channel().get_extendable(mount)
-        yield from (x for x,_ in extendable.read_extensions())
-
-    def get_extension_type(self, mount, name):
-        extendable = self.get_channel().get_extendable(mount)
-        extension = extendable.get_extension_class(name)
-        return extension.get_class_type()
-
-    def read_extension_types(self, mount):
-        extendable = self.get_channel().get_extendable(mount)
-        return extendable.read_extensions()
-
-    def create_extension(self, mount, name, *ctor_args, **ctor_kwargs):
-        extendable = self.get_channel().get_extendable(mount)
-        extension = extendable.get_extension_class(name)
-        return extension.construct(*ctor_args, **ctor_kwargs)
+class _ArgOverrider(object):
+    def _apply_arg_overrides(self):
+        try:
+            self._apply_guarded()
+        except Exception as e:
+            print("ArgOverrider:", "ERROR:", str(e))
+
+    def _apply_guarded(self):
+        if not self._invoke_path:
+            return
+
+        header = "Argument overrides"
+        for arg_name, arg_value in self.args:
+            if not self.args.is_default(arg_name):
+                continue
+
+            over_value = os.getenv(f"ushell{self._invoke_path}:{arg_name}")
+            if over_value is None:
+                continue
+
+            if header:
+                self.print_info(header)
+                header = None
+
+            print(arg_name + ": ", end="")
+            try:
+                arg_type = self.args.get_type(arg_name)
+                if isinstance(arg_value, tuple):
+                    import shlex
+                    over_value = (arg_type(x) for x in shlex.split(over_value))
+                    over_value = (*arg_value, *over_value)
+                    print(*over_value, end="")
+                else:
+                    over_value = arg_type(over_value)
+                    print(over_value, end="")
+            except Exception as e:
+                print("[ERROR:", str(e))
+            print(" (env)")
+
+            setattr(self.args, arg_name, over_value)
 
 
 
@@ -138,7 +160,7 @@ class _Channel(object):
 from . import _flick
 Arg = _flick.Arg
 Opt = _flick.Opt
-class Cmd(_flick.Cmd, _Channel):
+class Cmd(_flick.Cmd, _ArgOverrider):
     Arg = _flick.Arg
     Opt = _flick.Opt
 
@@ -162,6 +184,8 @@ class Cmd(_flick.Cmd, _Channel):
         if "PWD" in os.environ:
             del os.environ["PWD"]
 
+        self._apply_arg_overrides()
+
         try:
             return super()._call_main()
         except KeyboardInterrupt as e:
diff --git a/Engine/Extras/ushell/channels/flow/core/system/flow/system.py b/Engine/Extras/ushell/channels/flow/core/system/flow/system.py
--- a/Engine/Extras/ushell/channels/flow/core/system/flow/system.py
+++ b/Engine/Extras/ushell/channels/flow/core/system/flow/system.py
@@ -69,20 +69,26 @@ class _CommandClass(_LoadableClass):
 class _CommandNode(_FlowObject):
     _MAGIC = 493
 
-    def __init__(self, system, node):
+    def __init__(self, system, path, node):
         super().__init__(system)
+        self._path = path
         self._node = node
 
+    def get_path(self):
+        return self._path
+
     def read_children(self):
         system = self.get_system()
         for name, child_node in self._node.items():
             if isinstance(name, str):
-                yield name, _CommandNode(system, child_node)
+                path = self.get_path() + "/" + name
+                yield name, _CommandNode(system, path, child_node)
 
     def find_child(self, name):
         child = self._node.get(name, None)
         if child:
-            return _CommandNode(self.get_system(), child)
+            path = self.get_path() + "/" + name
+            return _CommandNode(self.get_system(), path, child)
 
     def get_command_class(self):
         specs = self._node.get(_CommandNode._MAGIC, None)
@@ -99,14 +105,16 @@ class _CommandTree(_FlowObject):
         self._root_node = root_node
 
     def get_root_node(self):
-        return _CommandNode(self.get_system(), self._root_node)
+        return _CommandNode(self.get_system(), "", self._root_node)
 
     def find(self, arg_iter):
+        path = ""
         arg = None
         node = self._root_node
         for arg in arg_iter:
             try:
                 node = node[arg]
+                path += "/" + arg
             except KeyError:
                 break
         else:
@@ -116,30 +124,7 @@ class _CommandTree(_FlowObject):
             if arg != None: yield arg
             yield from arg_iter
 
-        return _CommandNode(self.get_system(), node), read_ret_args()
-
-
-
-#-------------------------------------------------------------------------------
-class _ExtensionClass(_LoadableClass):
-    pass
-
-#-------------------------------------------------------------------------------
-class _Extendable(_FlowObject):
-    def __init__(self, system, data):
-        super().__init__(system)
-        self._data = data
-
-    def read_extensions(self):
-        system = self.get_system()
-        yield from ((k, _ExtensionClass(system, k, v)) for k,v in self._data.items())
-
-    def get_extension_class(self, name):
-        specs = self._data.get(name, None)
-        if specs:
-            return _ExtensionClass(self.get_system(), "ext." + name, specs)
-
-        raise ValueError(f"Unknown extendable '{name}'")
+        return _CommandNode(self.get_system(), path, node), read_ret_args()
 
 
 
@@ -181,17 +166,6 @@ class _Channel(_FlowObject):
         if parent >= 0:
             return self.get_system().get_channel(parent)
 
-    def get_extendable(self, name):
-        extendable = self._data["extendables"].get(name, None)
-        if extendable:
-            return _Extendable(self.get_system(), extendable)
-
-        parent = self.get_parent()
-        if parent:
-            return parent.get_extendable(name)
-
-        raise ValueError(f"Unknown extendable '{name}'")
-
     def read_tools(self):
         for name, tool in self._data["tools"].items():
             yield _Tool(self.get_system(), name, tool)
diff --git a/Engine/Extras/ushell/channels/flow/core/system/lib/bootstrap.py b/Engine/Extras/ushell/channels/flow/core/system/lib/bootstrap.py
--- a/Engine/Extras/ushell/channels/flow/core/system/lib/bootstrap.py
+++ b/Engine/Extras/ushell/channels/flow/core/system/lib/bootstrap.py
@@ -225,8 +225,16 @@ def _manifest_tool(name, tool):
 #-------------------------------------------------------------------------------
 def _install_tool(name, tool, channel_dir):
     manifest = _manifest_tool(name, tool)
-    if tool._platform and tool._platform != sys.platform:
-        return manifest
+
+    if req_platform := tool._platform:
+        this_platform = sys.platform
+        if "-" in req_platform:
+            import platform
+            arch = platform.machine()
+            arch = "amd64" if arch in ("x86_64", "AMD64") else "arm64"
+            this_platform += "-" + arch
+        if req_platform != this_platform:
+            return manifest
 
     tool_double = f"{name}-{tool._version}/"
     _log.write("Installing tool", tool_double)
diff --git a/Engine/Extras/ushell/channels/flow/core/system/run.py b/Engine/Extras/ushell/channels/flow/core/system/run.py
--- a/Engine/Extras/ushell/channels/flow/core/system/run.py
+++ b/Engine/Extras/ushell/channels/flow/core/system/run.py
@@ -55,7 +55,8 @@ def main(args):
     command_class = tree_node.get_command_class()
     if command_class:
         command = command_class.construct()
-        return command.invoke(tail_args)
+        command_path = tree_node.get_path()
+        return command.invoke(tail_args, invoke_path=command_path)
 
     # We've reached a point in the command tree that has no command. Inform the
     # user about the available sub-commands.
diff --git a/Engine/Extras/ushell/channels/flow/nt/describe.flow.py b/Engine/Extras/ushell/channels/flow/nt/describe.flow.py
--- a/Engine/Extras/ushell/channels/flow/nt/describe.flow.py
+++ b/Engine/Extras/ushell/channels/flow/nt/describe.flow.py
@@ -21,9 +21,9 @@ fd.source("[REDACTED_PROFILE_URL]/fd/releases/latest", "fd-v([0-9.]+)-x86_64
 
 #-------------------------------------------------------------------------------
 fzf_win32 = flow.describe.Tool()
-fzf_win32.version("0.44.1")
-fzf_win32.payload("[REDACTED_PROFILE_URL]/fzf/releases/download/$VERSION/fzf-$VERSION-windows_amd64.zip")
-fzf_win32.sha1("463f064d0403714ab5bf68ebf9fdea4c0159f6c3")
+fzf_win32.version("0.54.3")
+fzf_win32.payload("[REDACTED_PROFILE_URL]/fzf/releases/download/v$VERSION/fzf-$VERSION-windows_amd64.zip")
+fzf_win32.sha1("2c3ceab1a099ad3ab30eaffb6edaaa8f1a8ce2b4")
 fzf_win32.platform("win32")
 fzf_win32.bin("fzf.exe")
 fzf_win32.source("[REDACTED_PROFILE_URL]/fzf/releases/latest", r"fzf-(\d+\.\d+\.\d+)-windows")
diff --git a/Engine/Extras/ushell/channels/unreal/core/describe.flow.py b/Engine/Extras/ushell/channels/unreal/core/describe.flow.py
--- a/Engine/Extras/ushell/channels/unreal/core/describe.flow.py
+++ b/Engine/Extras/ushell/channels/unreal/core/describe.flow.py
@@ -5,43 +5,57 @@ import flow.describe
 
 #-------------------------------------------------------------------------------
 fzf_linux = flow.describe.Tool()
-fzf_linux.version("0.44.1")
-fzf_linux.payload("[REDACTED_PROFILE_URL]/fzf/releases/download/$VERSION/fzf-$VERSION-linux_amd64.tar.gz")
-fzf_linux.sha1("e7780f1e5e2dd4c8aa5e94dcbcbedf4cebfaaceb")
+fzf_linux.version("0.54.3")
+fzf_linux.payload("[REDACTED_PROFILE_URL]/fzf/releases/download/v$VERSION/fzf-$VERSION-linux_amd64.tar.gz")
+fzf_linux.sha1("e46fcd29c453ee61be0d41ca22b7eefd944f91f6")
 fzf_linux.platform("linux")
 fzf_linux.bin("fzf")
 
-fzf_darwin = flow.describe.Tool()
-fzf_darwin.version("0.44.1")
-fzf_darwin.payload("[REDACTED_PROFILE_URL]/fzf/releases/download/$VERSION/fzf-$VERSION-darwin_amd64.zip")
-fzf_darwin.sha1("405156eb6fc3d5275774cc569e24feaed362eab5")
-fzf_darwin.platform("darwin")
-fzf_darwin.bin("fzf")
+fzf_darwin_x64 = flow.describe.Tool()
+fzf_darwin_x64.version("0.54.3")
+fzf_darwin_x64.payload("[REDACTED_PROFILE_URL]/fzf/releases/download/v$VERSION/fzf-$VERSION-darwin_amd64.tar.gz")
+fzf_darwin_x64.sha1("3723b3445037bdaa302d1d365d4940beb4c283c8")
+fzf_darwin_x64.platform("darwin-amd64")
+fzf_darwin_x64.bin("fzf")
+
+fzf_darwin_arm = flow.describe.Tool()
+fzf_darwin_arm.version("0.54.3")
+fzf_darwin_arm.payload("[REDACTED_PROFILE_URL]/fzf/releases/download/v$VERSION/fzf-$VERSION-darwin_arm64.tar.gz")
+fzf_darwin_arm.sha1("c52657e62c095909fb1ac69e6d714eac3a10dc4c")
+fzf_darwin_arm.platform("darwin-arm64")
+fzf_darwin_arm.bin("fzf")
 
 
 
 #-------------------------------------------------------------------------------
 ripgrep_win32 = flow.describe.Tool()
-ripgrep_win32.version("14.0.3")
+ripgrep_win32.version("14.1.0")
 ripgrep_win32.payload("[REDACTED_PROFILE_URL]/ripgrep/releases/download/$VERSION/ripgrep-$VERSION-x86_64-pc-windows-msvc.zip")
-ripgrep_win32.sha1("5508b3dd5c12206c5d1c7994490bacf317cee0aa")
+ripgrep_win32.sha1("ba7820c05aeb49050d9242a109e422585c3fc767")
 ripgrep_win32.platform("win32")
 ripgrep_win32.bin("rg.exe")
 ripgrep_win32.source("[REDACTED_PROFILE_URL]/ripgrep/releases/latest", r"ripgrep-(\d+\.\d+\.\d+)-x86_64")
 
 ripgrep_linux = flow.describe.Tool()
-ripgrep_linux.version("14.0.3")
+ripgrep_linux.version("14.1.0")
 ripgrep_linux.payload("[REDACTED_PROFILE_URL]/ripgrep/releases/download/$VERSION/ripgrep-$VERSION-x86_64-unknown-linux-musl.tar.gz")
-ripgrep_linux.sha1("25a759834819f91625ac739c890a4b0139df7b63")
+ripgrep_linux.sha1("8267841999fa96f5a5a375ed518a24bebef6ca4c")
 ripgrep_linux.platform("linux")
 ripgrep_linux.bin("rg")
 
-ripgrep_darwin = flow.describe.Tool()
-ripgrep_darwin.version("14.0.3")
-ripgrep_darwin.payload("[REDACTED_PROFILE_URL]/ripgrep/releases/download/$VERSION/ripgrep-$VERSION-x86_64-apple-darwin.tar.gz")
-ripgrep_darwin.sha1("d5af393ea9b5a443544edc686cf939f517be35ea")
-ripgrep_darwin.platform("darwin")
-ripgrep_darwin.bin("rg")
+ripgrep_darwin_x64 = flow.describe.Tool()
+ripgrep_darwin_x64.version("14.1.0")
+ripgrep_darwin_x64.payload("[REDACTED_PROFILE_URL]/ripgrep/releases/download/$VERSION/ripgrep-$VERSION-x86_64-apple-darwin.tar.gz")
+ripgrep_darwin_x64.sha1("0d98d699f08c4eed18f4379734d4e1d973f59cde")
+ripgrep_darwin_x64.platform("darwin-amd64")
+ripgrep_darwin_x64.bin("rg")
+
+ripgrep_darwin_arm = flow.describe.Tool()
+ripgrep_darwin_arm.version("14.1.0")
+ripgrep_darwin_arm.payload("[REDACTED_PROFILE_URL]/ripgrep/releases/download/$VERSION/ripgrep-$VERSION-aarch64-apple-darwin.tar.gz")
+ripgrep_darwin_arm.sha1("59d0d543a7e25831b29ebe2943da366d4c005d26")
+ripgrep_darwin_arm.platform("darwin-amd64")
+ripgrep_darwin_arm.bin("rg")
 
 
 
