Updated ushell [35652478]

#rnx

diff --git a/Engine/Extras/ushell/channels/flow/core/system/flow/cmd.py b/Engine/Extras/ushell/channels/flow/core/system/flow/cmd.py
--- a/Engine/Extras/ushell/channels/flow/core/system/flow/cmd.py
+++ b/Engine/Extras/ushell/channels/flow/core/system/flow/cmd.py
@@ -126,8 +126,9 @@ class _ArgOverrider(object):
 
         header = "Argument overrides"
         for arg_name, arg_value in self.args:
-            if not self.args.is_default(arg_name):
-                continue
+            if not isinstance(arg_value, tuple):
+                if not self.args.is_default(arg_name):
+                    continue
 
             over_value = os.getenv(f"ushell{self._invoke_path}:{arg_name}")
             if over_value is None:
