Updated ushell [35689748]

#rnx

diff --git a/Engine/Extras/ushell/UgsTool.json b/Engine/Extras/ushell/UgsTool.json
deleted file mode 100644
diff --git a/Engine/Extras/ushell/channels/flow/core/system/flow/_flick.py b/Engine/Extras/ushell/channels/flow/core/system/flow/_flick.py
--- a/Engine/Extras/ushell/channels/flow/core/system/flow/_flick.py
+++ b/Engine/Extras/ushell/channels/flow/core/system/flow/_flick.py
@@ -182,16 +182,24 @@ class Cmd(object):
         col_o = max(len(x[0]) for x in opts)
         col_0 = max(col_a, col_o)
 
-        once = True
-        for name, desc in args:
-            if once:
-                print("\nARGS:")
-                once = False
-            print("  %-*s" % (col_0, name), desc)
+        import textwrap
+        width = max(40, 80 - col_0 - 4) # 3 = leader + separator = 2 + 2
+        def print_param(name, desc):
+            print("  %-*s  " % (col_0, name), end="")
+            leader = ""
+            wrapper = textwrap.TextWrapper(width=width, subsequent_indent="  ")
+            for line in wrapper.wrap(desc):
+                print(leader, line, sep="")
+                leader = leader or " " * (col_0 + 3)
+
+        if args:
+            print("\nARGS:")
+            for name, desc in args:
+                print_param(name, desc)
 
         print("\nOPTIONS:")
         for name, desc in opts:
-            print("  %-*s" % (col_0, name), desc)
+            print_param(name, desc)
 
     def _call_main(self):
         ret = self.main()
@@ -213,6 +221,9 @@ class Cmd(object):
 
         # Parses an opt-type from input stream
         def parse_opt(arg, in_args_iter):
+            if not arg.startswith("--"):
+                raise ValueError((arg,))
+
             # Parse "--opt_name[=value]"
             opt_name, value = (*arg[2:].split("=", 1), None)[:2]
 
@@ -239,8 +250,8 @@ class Cmd(object):
         positionals = []
         in_args_iter = iter(in_args)
         for arg in in_args_iter:
-            if arg.startswith("--"):
-                if len(arg) == 2:
+            if arg.startswith("-"):
+                if arg == "--":
                     positionals.append("--")
                     positionals.extend(in_args_iter)
                     break
