Deleted file in error due to an automated horror show
#rnx

diff --git a/Engine/Extras/ushell/UgsTool.json b/Engine/Extras/ushell/UgsTool.json
new file mode 100644
--- /dev/null
+++ b/Engine/Extras/ushell/UgsTool.json
@@ -0,0 +1,10 @@
+{
+	"statusPanelLinks": [
+		{
+			"label": "ushell",
+			"fileName": "$(COMSPEC)",
+			"arguments": "/D /K \"\"$(ToolDir)\\ushell.bat\" --project=\"$(ProjectFile)\"\"",
+			"workingDir": "$(ProjectDir)"
+		}
+	]
+}
