Updated ushell [35736575]

#rnx

diff --git a/Engine/Extras/ushell/channels/flow/core/system/flow/_flick.py b/Engine/Extras/ushell/channels/flow/core/system/flow/_flick.py
--- a/Engine/Extras/ushell/channels/flow/core/system/flow/_flick.py
+++ b/Engine/Extras/ushell/channels/flow/core/system/flow/_flick.py
@@ -190,7 +190,7 @@ class Cmd(object):
             wrapper = textwrap.TextWrapper(width=width, subsequent_indent="  ")
             for line in wrapper.wrap(desc):
                 print(leader, line, sep="")
-                leader = leader or " " * (col_0 + 3)
+                leader = leader or (" " * (col_0 + 4))
 
         if args:
             print("\nARGS:")
diff --git a/Engine/Extras/ushell/channels/unreal/core/pylib/unreal/_context.py b/Engine/Extras/ushell/channels/unreal/core/pylib/unreal/_context.py
--- a/Engine/Extras/ushell/channels/unreal/core/pylib/unreal/_context.py
+++ b/Engine/Extras/ushell/channels/unreal/core/pylib/unreal/_context.py
@@ -440,11 +440,17 @@ class _Project(_UfsNode):
 
     def _create_temp_target(self, target:_Target) -> None:
         int_dir = self.get_dir() / "Source"
-        int_dir.mkdir(parents=True, exist_ok=True)
 
-        type_name = target.get_type().name.title()
-        target_name = target.get_name()
+        # Only do creation if we are the first to arrive
+        int_tag = int_dir / ".ushell_created"
+        if not int_dir.is_dir():
+            int_dir.mkdir(parents=True, exist_ok=True)
+            int_tag.open("wb").close()
+        else:
+            if not int_tag.is_file():
+                return
 
+        target_name = target.get_name()
         dot_cs_path = int_dir / (target_name + ".Target.cs")
         if dot_cs_path.is_file():
             return
@@ -453,6 +459,7 @@ class _Project(_UfsNode):
         if self._engine.get_version_major() <= 4:
             game_target = "UE4Game"
 
+        type_name = target.get_type().name.title()
         with dot_cs_path.open("wt") as dot_cs:
             dot_cs.write(
                 "// Created by ushell to make blueprint projects work end to end.\n"
diff --git a/Engine/Extras/ushell/channels/unreal/perforce/cmds/clean.py b/Engine/Extras/ushell/channels/unreal/perforce/cmds/clean.py
--- a/Engine/Extras/ushell/channels/unreal/perforce/cmds/clean.py
+++ b/Engine/Extras/ushell/channels/unreal/perforce/cmds/clean.py
@@ -192,7 +192,9 @@ class Clean(flow.cmd.Cmd):
         self.print_info("Finding directories and files to clean")
         root_dirs = [
             engine_dir,
+            *(x for x in engine_dir.glob("Platforms/*") if x.is_dir()),
             *(x for x in engine_dir.glob("Programs/*") if x.is_dir()),
+            *(x for x in engine_dir.glob("Restricted/*") if x.is_dir()),
         ]
 
         print("Enumerating...", end="")
diff --git a/Engine/Extras/ushell/channels/unreal/perforce/cmds/sync.py b/Engine/Extras/ushell/channels/unreal/perforce/cmds/sync.py
--- a/Engine/Extras/ushell/channels/unreal/perforce/cmds/sync.py
+++ b/Engine/Extras/ushell/channels/unreal/perforce/cmds/sync.py
@@ -323,14 +323,16 @@ class Sync(_SyncBase):
         def impl(path):
             print("Source:", os.path.normpath(path), end="")
             try:
-                sync_config = open(path, "rt")
-                print()
+                with open(path, "rt") as sync_config:
+                    lines = [x.strip() for x in sync_config]
+                print(" ... ", len(lines), "lines")
             except:
                 print(" ... not found")
                 return
 
+            # Exclusions
             def read_exclusions():
-                for line in map(str.strip, sync_config):
+                for line in lines:
                     if line.startswith("-"):    yield line[1:]
                     elif line.startswith("$-"): yield line[2:]
 
@@ -350,7 +352,15 @@ class Sync(_SyncBase):
                     view = view or line
                     print(flow.cmd.text.light_yellow(view + " (ill-formed)"))
 
-            sync_config.close()
+            # Extra roots
+            def read_extra_roots():
+                for line in lines:
+                    if line.startswith("/"):    yield line
+                    elif line.startswith("$/"): yield line[1:]
+
+            for extra_root in read_extra_roots():
+                syncer.add_path(self._local_root + extra_root)
+                print("        +", "sync   ", extra_root)
 
         self.print_info("Applying .p4sync.txt")
         for dir in (self.get_home_dir(), self._local_root):
diff --git a/Engine/Extras/ushell/channels/unreal/perforce/pylib/p4utils.py b/Engine/Extras/ushell/channels/unreal/perforce/pylib/p4utils.py
--- a/Engine/Extras/ushell/channels/unreal/perforce/pylib/p4utils.py
+++ b/Engine/Extras/ushell/channels/unreal/perforce/pylib/p4utils.py
@@ -149,8 +149,13 @@ def get_branch_root(depot_path):
     for x in fstat_paths():
         print(" ", x)
 
+    def on_error(message):
+        if data := getattr(message, "data", None):
+            if "has been unloaded" in data:
+                raise EnvironmentError("Attempting to use an unloaded client")
+
     fstat = P4.fstat(fstat_paths(), T="depotFile")
-    root_path = fstat.run(on_error=False)
+    root_path = fstat.run(on_error=on_error)
     if root_path:
         return "/".join(root_path.depotFile.split("/")[:-1]) + "/"
 
