Updated ushell [35760554]

#rnx

diff --git a/Engine/Extras/ushell/channels/unreal/core/cmds/stage.py b/Engine/Extras/ushell/channels/unreal/core/cmds/stage.py
--- a/Engine/Extras/ushell/channels/unreal/core/cmds/stage.py
+++ b/Engine/Extras/ushell/channels/unreal/core/cmds/stage.py
@@ -179,6 +179,5 @@ class Deploy(_Impl):
         self.args.build  = False
         self.args.cook   = False
         self.args.deploy = True
-        self.args.nopak  = True
-        self.args.zen    = False
+        self.args.style  = 'nopak'
         return super().main(skipstage=True)
