Updated ushell [35864467]

#rnx

diff --git a/Engine/Extras/ushell/channels/unreal/core/cmds/stage.py b/Engine/Extras/ushell/channels/unreal/core/cmds/stage.py
--- a/Engine/Extras/ushell/channels/unreal/core/cmds/stage.py
+++ b/Engine/Extras/ushell/channels/unreal/core/cmds/stage.py
@@ -179,5 +179,5 @@ class Deploy(_Impl):
         self.args.build  = False
         self.args.cook   = False
         self.args.deploy = True
-        self.args.style  = 'nopak'
+        self.args.style  = "nopak"
         return super().main(skipstage=True)
diff --git a/Engine/Extras/ushell/channels/unreal/core/pylib/cipher.py b/Engine/Extras/ushell/channels/unreal/core/pylib/cipher.py
--- a/Engine/Extras/ushell/channels/unreal/core/pylib/cipher.py
+++ b/Engine/Extras/ushell/channels/unreal/core/pylib/cipher.py
@@ -55,6 +55,8 @@ def _hash(x):
     return ret
 
 
+TOC_SIZE        = 8192
+TOC_ENTRY_SIZE  = 128
 
 #-------------------------------------------------------------------------------
 class Blob(object):
@@ -68,12 +70,12 @@ class Blob(object):
             data = memoryview(data)
         self._set_data(data)
         self._ranges = []
-        self._toc_set = [0] * (len(self._toc) // 128)
+        self._toc_set = [0] * (len(self._toc) // TOC_ENTRY_SIZE)
 
     def _set_data(self, data):
         self._data = data.cast("B")
         self._lede = data[:256]
-        self._toc = data[256:4096]
+        self._toc = data[256:TOC_SIZE]
 
     def get_data(self) -> memoryview:
         return self._data
@@ -86,8 +88,8 @@ class Blob(object):
 
     def find(self, key:memoryview) -> Iterator[Tuple[str, memoryview]]:
         key, nonce, nonce_toc = self._diffuse(key)
-        for i in range(0, len(self._toc), 128):
-            piece = self._toc[i:i + 128]
+        for i in range(0, len(self._toc), TOC_ENTRY_SIZE):
+            piece = self._toc[i:i + TOC_ENTRY_SIZE]
             piece = _chacha20(piece, key=key, nonce=nonce_toc)
             piece = memoryview(piece).cast("L")
             name = piece[4:].cast("B")
@@ -106,11 +108,11 @@ class Blob(object):
 
         # find somewhere to store it or grow if we fail a lot
         offset = 0
-        space = len(self._data) - size - 4096
+        space = len(self._data) - size - TOC_SIZE
         assert space > 0
         for i in range(100):
             offset = secrets.randbelow(space)
-            offset += 4096
+            offset += TOC_SIZE
             for l,r in self._ranges:
                 if not (offset + size <= l or offset >= r):
                     offset = 0
@@ -137,7 +139,7 @@ class Blob(object):
         assert name_length <= 112
         name = name + secrets.token_bytes(112 - name_length)
 
-        toc_entry = bytearray(128)
+        toc_entry = bytearray(TOC_ENTRY_SIZE)
         toc_entry[16:] = name
         toc_entry = memoryview(toc_entry).cast("L")
         toc_entry[0] = offset
@@ -153,8 +155,10 @@ class Blob(object):
             if self._toc_set[toc_index] == 0:
                 self._toc_set[toc_index] = 1
                 break
+            if 0 not in self._toc_set:
+                raise RuntimeError("toc is full")
 
         # blit encrypted parts into blob
-        toc_index = toc_index * 128
+        toc_index = toc_index * TOC_ENTRY_SIZE
         self._data[offset:offset + size] = data
-        self._toc[toc_index:toc_index + 128] = toc_entry
+        self._toc[toc_index:toc_index + TOC_ENTRY_SIZE] = toc_entry
diff --git a/Engine/Extras/ushell/channels/unreal/perforce/cmds/gui.py b/Engine/Extras/ushell/channels/unreal/perforce/cmds/gui.py
--- a/Engine/Extras/ushell/channels/unreal/perforce/cmds/gui.py
+++ b/Engine/Extras/ushell/channels/unreal/perforce/cmds/gui.py
@@ -25,7 +25,6 @@ _check_for_existing = _check_for_existing_nt if os.name == "nt" else lambda x: "
 #-------------------------------------------------------------------------------
 class Gui(flow.cmd.Cmd):
     """ Opens the clientspec for the current directory in P4V """
-    filename = flow.cmd.Arg(Path(), "Path to select in the workspace/depot tree")
     p4vargs  = flow.cmd.Arg([str], "Additional arguments to pass to P4V")
     def main(self):
         self.print_info("Fetching Perforce info")
@@ -62,10 +61,8 @@ class Gui(flow.cmd.Cmd):
             "-c", client,
         )
 
-        if self.args.filename.is_file():
-            filename = self.args.filename.resolve()
-            args = (*args, "-s", str(filename))
+        p4v_cwd = Path(os.getenv("TEMP", client_dir)).resolve()
 
         print("p4v", *args)
-        subprocess.Popen(("p4v", *args, *self.args.p4vargs), cwd=client_dir)
+        subprocess.Popen(("p4v", *args, *self.args.p4vargs), cwd=p4v_cwd)
         print("Done!")
diff --git a/Engine/Extras/ushell/channels/unreal/perforce/cmds/sync.py b/Engine/Extras/ushell/channels/unreal/perforce/cmds/sync.py
--- a/Engine/Extras/ushell/channels/unreal/perforce/cmds/sync.py
+++ b/Engine/Extras/ushell/channels/unreal/perforce/cmds/sync.py
@@ -320,17 +320,19 @@ class Sync(_SyncBase):
                         x.write(line)
 
     def _apply_p4sync_txt(self, syncer):
-        def impl(path):
-            print("Source:", os.path.normpath(path), end="")
+        view_filter = syncer.get_view_filter()
+
+        def load_lines(index, path):
+            print("Source:", index, os.path.normpath(path), end="")
             try:
                 with open(path, "rt") as sync_config:
                     lines = [x.strip() for x in sync_config]
                 print(" ... ", len(lines), "lines")
+                return lines
             except:
                 print(" ... not found")
-                return
 
-            # Exclusions
+        def impl_exclusions(index, lines):
             def read_exclusions():
                 for line in lines:
                     if line.startswith("-"):    yield line[1:]
@@ -342,29 +344,46 @@ class Sync(_SyncBase):
                 elif line.startswith("/"):   view = line[1:]
                 elif line.startswith("..."): view = line
 
-                print("       %2d" % i, "exclude", end=" ")
+                print("  %d.%02d" % (index, i), "excl", end=" ")
 
                 if view and (view.count("/") or "/*." in view or view.startswith("*.")):
                     view = self._branch_root + view
-                    syncer.add_exclude(view)
+                    view_filter.add_exclude(view)
                     print(view)
                 else:
                     view = view or line
                     print(flow.cmd.text.light_yellow(view + " (ill-formed)"))
 
-            # Extra roots
+        def impl_extra_roots(index, lines):
+            view_query = view_filter.get_query()
+
             def read_extra_roots():
                 for line in lines:
-                    if line.startswith("/"):    yield line
-                    elif line.startswith("$/"): yield line[1:]
-
-            for extra_root in read_extra_roots():
-                syncer.add_path(self._local_root + extra_root)
-                print("        +", "sync   ", extra_root)
+                    if line.startswith("/"):    yield line[1:]
+                    elif line.startswith("$/"): yield line[2:]
+
+            for i, extra_root in enumerate(read_extra_roots()):
+                extra_root = self._branch_root + extra_root
+                syncer.add_path(extra_root)
+                print("  %d.%02d" % (index, i), "sync", extra_root)
+
+                if result := view_query.is_excluded(extra_root):
+                    print(
+                        " " * 11,
+                        flow.cmd.text.light_yellow(f"ignoring exclusion '{result}'"),
+                        "(contradiction)"
+                    )
+                    view_filter.remove_exclude(result)
 
         self.print_info("Applying .p4sync.txt")
-        for dir in (self.get_home_dir(), self._local_root):
-            impl(dir + ".p4sync.txt")
+
+        liness = []
+        for index, dir in enumerate((self.get_home_dir(), self._local_root)):
+            if lines := load_lines(index, dir + ".p4sync.txt"):
+                liness.append((index, lines))
+
+        for index, lines in liness: impl_exclusions(index, lines)
+        for index, lines in liness: impl_extra_roots(index, lines)
 
 
 
diff --git a/Engine/Extras/ushell/channels/unreal/perforce/pylib/p4utils.py b/Engine/Extras/ushell/channels/unreal/perforce/pylib/p4utils.py
deleted file mode 100644
diff --git a/Engine/Extras/ushell/channels/unreal/perforce/pylib/p4utils/__init__.py b/Engine/Extras/ushell/channels/unreal/perforce/pylib/p4utils/__init__.py
new file mode 100644
--- /dev/null
+++ b/Engine/Extras/ushell/channels/unreal/perforce/pylib/p4utils/__init__.py
@@ -0,0 +1,164 @@
+# Copyright Epic Games, Inc. All Rights Reserved.
+
+import os
+from peafour import P4
+import subprocess as sp
+from pathlib import Path
+
+from .syncer import Syncer
+
+#-------------------------------------------------------------------------------
+p4_set_result = None
+def get_p4_set(prop_name):
+    if ret := os.getenv(prop_name):
+        return ret
+
+    global p4_set_result
+    if not p4_set_result:
+        p4_set_result = {}
+        try: proc = sp.Popen(("p4", "set", "-q"), stdout=sp.PIPE)
+        except: return
+
+        for line in iter(proc.stdout.readline, b""):
+            try: key, value = line.split(b"=", 1)
+            except ValueError: continue
+            p4_set_result[key.decode()] = value.strip().decode()
+
+        proc.wait()
+        proc.stdout.close()
+
+    return p4_set_result.get(prop_name, None)
+
+#-------------------------------------------------------------------------------
+def login():
+    try:
+        detail = P4.login(s=True).run()
+    except P4.Error:
+        raise EnvironmentError("No valid Perforce session found. Run 'p4 login' to authenticate.")
+    return getattr(detail, "User", None)
+
+#-------------------------------------------------------------------------------
+def get_p4config_name():
+    return os.path.basename(get_p4_set("P4CONFIG") or ".p4config.txt")
+
+#-------------------------------------------------------------------------------
+def has_p4config(start_dir):
+    from pathlib import Path
+    p4config = get_p4config_name()
+    for dir in (Path(start_dir) / "x").parents:
+        candidate = dir / p4config
+        if os.path.isfile(candidate):
+            return True, candidate
+    return False, p4config
+
+#-------------------------------------------------------------------------------
+def create_p4config(p4config_path, client, username, port=None):
+    with open(p4config_path, "wt") as out:
+        print_args = { "sep" : "", "file" : out }
+        print("P4CLIENT=", client, **print_args)
+        print("P4USER=", username, **print_args)
+        if port:
+            print("P4PORT", port, **print_args)
+
+#-------------------------------------------------------------------------------
+def ensure_p4config(start_dir=None):
+    start_dir = start_dir or os.getcwd()
+    found, p4config_name = has_p4config(start_dir)
+    if found:
+        return p4config_name, False
+
+    username = login()
+
+    # Get the client for 'start_dir'
+    client = get_client_from_dir(start_dir, username)
+    if not client:
+        return
+    client, root_dir = client
+
+    # Now we know where to locate a p4config file
+    p4config_path = f"{root_dir}/{p4config_name}"
+    create_p4config(p4config_path, client, username)
+    return p4config_path, True
+
+#-------------------------------------------------------------------------------
+def get_client_from_dir(root_dir, username):
+    import socket
+    host_name = socket.gethostname().lower()
+
+    root_dir = Path(root_dir)
+
+    clients = (x for x in P4.clients(u=username) if x.Host.lower() == host_name)
+    for client in clients:
+        if root_dir.is_relative_to(client.Root):
+            return client.client, client.Root
+
+#-------------------------------------------------------------------------------
+def get_branch_root(depot_path):
+    depot_path = str(depot_path)
+
+    def fstat_paths():
+        limit = 5 # ...two of which are always required
+        query_path = "//"
+        for piece in depot_path[2:].split("/")[:limit]:
+            query_path += piece + "/"
+            yield query_path + "GenerateProjectFiles.bat"
+
+    print("Probing for well-known file:")
+    for x in fstat_paths():
+        print(" ", x)
+
+    def on_error(message):
+        if data := getattr(message, "data", None):
+            if "has been unloaded" in data:
+                raise EnvironmentError("Attempting to use an unloaded client")
+
+    fstat = P4.fstat(fstat_paths(), T="depotFile")
+    root_path = fstat.run(on_error=on_error)
+    if root_path:
+        return "/".join(root_path.depotFile.split("/")[:-1]) + "/"
+
+    raise ValueError("Unable to establish branch root")
+
+#-------------------------------------------------------------------------------
+def run_p4vc(*args):
+    try:
+        proc = sp.run(("p4vc", *args))
+    except FileNotFoundError:
+        proc = sp.run(("p4v", "-p4vc", *args))
+    return proc.returncode
+
+
+
+#-------------------------------------------------------------------------------
+class TempBranchSpec(object):
+    def __init__(self, use, username, from_path, to_path, ignore_streams=False):
+        import hashlib
+        id = hashlib.md5()
+        id.update(from_path.encode())
+        id.update(to_path.encode())
+        id = id.hexdigest()[:6]
+        self._name = f"{username}-ushell.{use}-{id}"
+
+        # To map between streams we need to extract the internal branchspec that
+        # Perforce builds. If from/to aren't related streams it will fail so we
+        # fallback to a conventional trivial branchspec.
+        try:
+            if ignore_streams:
+                raise P4.Error("")
+
+            branch = P4.branch(self._name, o=True, S=from_path[:-1], P=to_path[:-1])
+            result = branch.run()
+            spec = result.as_dict()
+        except P4.Error:
+            spec = {
+                "Branch" : self._name,
+                "View0"  : f"{from_path}... {to_path}...",
+            }
+
+        P4.branch(i=True).run(input_data=spec, on_error=False)
+
+    def __del__(self):
+        P4.branch(self._name, d=True).run()
+
+    def __str__(self):
+        return self._name
diff --git a/Engine/Extras/ushell/channels/unreal/perforce/pylib/p4utils/syncer.py b/Engine/Extras/ushell/channels/unreal/perforce/pylib/p4utils/syncer.py
new file mode 100644
--- /dev/null
+++ b/Engine/Extras/ushell/channels/unreal/perforce/pylib/p4utils/syncer.py
@@ -0,0 +1,253 @@
+import re
+import time
+import threading
+from peafour import P4
+from . import view
+
+#-------------------------------------------------------------------------------
+def _kb_string(value):
+    return format(value // 1024, ",") + "KB"
+
+#-------------------------------------------------------------------------------
+class _SyncRota(object):
+    class _Worker(object):
+        def __init__(self, id):
+            self.id = id
+            self.work_items = []
+            self.burden = 0
+            self.done_size = 0
+            self.done_items = 0
+            self.error = False
+
+    def __init__(self, changelist, worker_count):
+        # An extra worker is used for syncs to #0 (i.e. deletes). When using the
+        # -L argument a move will be two path#rev rota items. If they are actioned
+        # by the same p4 worker only one item gets synced as "p4 sync" ignores the
+        # other. So a move's two path#rev actions are give to two separate workers.
+        self._workers = [_SyncRota._Worker(x) for x in range(worker_count + 1)]
+        self.changelist = str(changelist)
+
+    def add_work(self, item, rev, cost):
+        if rev == 0:
+            worker = self._workers[0]
+        else:
+            worker = min(self._workers[1:], key=lambda x: x.burden)
+        worker.work_items.append((item, rev, cost))
+        worker.burden += cost
+
+    def sort(self):
+        direction = 1
+        for worker in self._workers:
+            worker.work_items.sort(key=lambda x: x[2] * direction)
+            direction *= -1
+
+    def read_workers(self):
+        yield from (x for x in self._workers if x.work_items)
+
+#-------------------------------------------------------------------------------
+class Syncer(object):
+    def __init__(self):
+        super().__init__()
+        self._paths = []
+        self._view_filter = view.ViewFilter()
+
+    def _read_sync_specs(self, include_excluded=True):
+        cl_suffix = "@" + self._rota.changelist
+        cl_suffix = "#0" if cl_suffix == "@0" else cl_suffix
+        for depot_path in self._paths:
+            yield depot_path + cl_suffix
+
+        if include_excluded:
+            # Using "@0" results in slow queries it seems
+            yield from (x + "#0" for x in self.get_view_filter().read_excludes())
+
+    def get_view_filter(self):
+        return self._view_filter
+
+    def add_path(self, dir):
+        self._paths.append(dir)
+
+    def schedule(self, changelist, worker_count=8):
+        view_query = self.get_view_filter().get_query()
+
+        # P4.<cmd> uses p4's Python-marshalled output (the -G option). However the
+        # "p4 sync -n" will report open files via a "info" message instead of a
+        # structured "stat" one. So we capture those and fake a reply
+        def read_items():
+            misreports = []
+            def on_info(data):
+                if m := re.match(r"\s*(//[^#]+)#(\d+).+is opened and", data.data):
+                    misreports.append((m.group(1), m.group(2)))
+
+            yield from P4.sync(self._read_sync_specs(), n=True).read(on_error=False, on_info=on_info)
+
+            class FakeItem(object): pass
+            for depot_path, rev in misreports:
+                item = FakeItem()
+                item.depotFile = depot_path
+                item.rev = rev
+                item.fileSize = 1 # to save looking this up we'll use a nominal size
+                item.action = "edit"
+                yield item
+
+        self._rota = _SyncRota(changelist, worker_count)
+
+        # Fill the rota
+        total_size = 0
+        count = 0
+        for item in read_items():
+            depot_path = item.depotFile
+            rev = int(item.rev)
+
+            if view_query.is_excluded(depot_path):
+                if item.action != "deleted":
+                    continue
+                rev = 0
+
+            if count % 17 == 0:
+                print("\r" + str(count), "files", f"({_kb_string(total_size)})", end="")
+
+            size = int(getattr(item, "fileSize", 0)) # deletes have no size attr
+
+            # When switching some files are metadata only changes and no content
+            # is transferred.
+            if item.action == "replaced":
+                size = 0
+
+            # When syncing backwards files need to be removed. Perforce has been
+            # seen reporting these correctly as a delete action but with a revision
+            # that is not 0. So we'll just force a rev=0
+            if item.action == "deleted" and rev >= 1:
+                rev = 0
+
+            self._rota.add_work(depot_path, rev, size)
+
+            total_size += size
+            count += 1
+        self._rota.sort()
+        print("\r" + str(count), "files", f"({_kb_string(total_size)})")
+
+    def sync(self, *, dryrun=False, echo=False):
+        # Sum up what we have to do
+        total_burden = sum(x.burden for x in self._rota.read_workers())
+        total_items = sum(len(x.work_items) for x in self._rota.read_workers())
+        lock = threading.Lock()
+
+        print(f"Fetching {_kb_string(total_burden)} in {total_items} files")
+
+        # Launch the worker threads
+        def sync_thread(worker):
+            def on_error(p4_error):
+                if "not enough space" in p4_error.data:
+                    worker.error = "Out of disk space"
+                    raise EOFError()
+                else:
+                    with lock:
+                        print("sync/error: ", p4_error.data, flush=True)
+
+            try:
+                # The 'p4 sync -L' option reportedly does less work on the server
+                # so that's what we use here. Note that it only works with "#rev"
+                # syntax if one wants to sync a specific revision or changelist.
+                def read_sync_items():
+                    for path, rev, size in worker.work_items:
+                        yield f"{path}#{rev}"
+
+                sync = P4(b=8192).sync(read_sync_items(), n=dryrun, L=True)
+                for item in sync.read(on_error=on_error):
+                    if echo:
+                        with lock:
+                            print("sync/worker:", item.depotFile, "#" + str(item.rev), flush=True)
+
+                    size = int(getattr(item, "fileSize", 0))
+                    if item.action == "replaced":
+                        # meta data update only, no content xfer
+                        size = 0
+
+                    worker.done_size += size + 0.01
+            except EOFError:
+                pass
+
+        def create_thread(worker):
+            thread = threading.Thread(target=sync_thread, args=(worker,))
+            thread.start()
+            return thread
+
+        threads = [create_thread(x) for x in self._rota.read_workers()]
+        print(f"Using {len(threads)} workers")
+
+        # While there are active threads, print detail about their progress
+        total_burden += (0.01 * total_items)
+        while not echo:
+            threads = [x for x in threads if x.is_alive()]
+            if not threads:
+                break
+
+            done_size = sum(x.done_size for x in self._rota.read_workers())
+            progress = ((done_size * 1000) // total_burden) / 10
+            print("\r%5.1f%%" % progress, _kb_string(int(done_size)), end="");
+            time.sleep(0.3)
+        else:
+            for thread in threads:
+                thread.join()
+        print("\r...done              ")
+
+        # Check for errors from the workers
+        for worker in (x for x in self._rota.read_workers() if x.error):
+            print(flow.cmd.text.red("!!" + str(worker.error)))
+            return False
+
+        # Nothing more to do if this is a dry run as the remaining tasks need a
+        # sync to operate on.
+        if dryrun:
+            return True
+
+        view_query = self.get_view_filter().get_query()
+
+        # P4.sync() returns 'stat' type events but "p4 sync" will report files
+        # with a complex sync scenario only as unstructured 'info' messages. As the
+        # above won't know about these files we'll do a second sync to catch them.
+        global sync_errors
+        sync_errors = False
+        print("Finalising...", end="")
+        def read_depot_files():
+            def on_error(data):
+                msg = data.data.strip()
+                if "up-to-date" in msg: return
+                if "not in client view" in msg: return
+                if "protected namespace - access denied" in msg: return
+
+                print("\n", flow.cmd.text.red(msg), end="")
+                global sync_errors
+                sync_errors = True
+
+            p4_is_daft = []
+            def on_info(data):
+                msg = data.data
+                if " - must resolve #" in msg: return
+                if " - is opened and not being changed" in msg:
+                    try: p4_is_daft.append(msg.split(" - ")[0])
+                    except IOError: pass
+                    return
+                print("\n", flow.cmd.text.light_yellow(data.data), end="")
+
+            sync = P4.sync(self._read_sync_specs(False), n=True)
+            for item in sync.read(on_error=on_error, on_info=on_info):
+                if not view_query.is_excluded(item.depotFile):
+                    yield item.depotFile + "#" + item.rev
+
+            yield from iter(p4_is_daft)
+
+        count = 0
+        sync = P4.sync(read_depot_files(), L=True)
+        for item in sync.read(on_error=False):
+            if echo:
+                with lock:
+                    print(f"sync/final: {item.depotFile} #{item.rev}")
+            else:
+                count += 1
+                print("\rFinalising...", count, end="")
+
+        print()
+
+        return not sync_errors
diff --git a/Engine/Extras/ushell/channels/unreal/perforce/pylib/p4utils/view.py b/Engine/Extras/ushell/channels/unreal/perforce/pylib/p4utils/view.py
new file mode 100644
--- /dev/null
+++ b/Engine/Extras/ushell/channels/unreal/perforce/pylib/p4utils/view.py
@@ -0,0 +1,79 @@
+import re
+
+#-------------------------------------------------------------------------------
+class _ViewWildcard(object):
+    def __init__(self, view, *, allow_empty_wildcards=False, case_sensitive=False):
+        self._src = view
+
+        view = view.replace("\\", "/")
+
+        if view.endswith("/"):
+            view += "\x01"
+        elif view.endswith("/..."):
+            view = view[:-3] + "\x01"
+
+        rooted = view.startswith("//")
+
+        regex = ""
+        sep = ""
+        for piece in (x for x in view.split("/") if x):
+            regex += sep
+            sep = "/"
+
+            if piece == "..." and allow_empty_wildcards:
+                regex += "(?:[^/]+/)*"
+                sep = ""
+                continue
+
+            piece = piece.replace("...", "\x01")
+            piece = piece.replace("*", "\x02")
+            piece = re.escape(piece)
+            piece = piece.replace("\x01", ".*")
+            piece = piece.replace("\x02", "[^/]*")
+            regex += piece
+
+        regex = "//" + regex if rooted else regex
+        regex = (regex or "^") + "$"
+        flags = None if case_sensitive else re.IGNORECASE
+        self._re = re.compile(regex, flags)
+
+    def get_view(self):
+        return self._src
+
+    def test(self, depot_path):
+        return self._re.search(depot_path) is not None
+
+#-------------------------------------------------------------------------------
+class ViewFilter(object):
+    class Query(object):
+        def __init__(self, view_excludes):
+            self._view_excludes = view_excludes
+
+        def is_excluded(self, path):
+            return next((x.get_view() for x in self._view_excludes if x.test(path)), False)
+
+    def __init__(self):
+        self._excluded_views = set()
+        self._query = None
+
+    def get_query(self):
+        if self._query:
+            return self._query
+        ret = [_ViewWildcard(x, allow_empty_wildcards=True) for x in self._excluded_views]
+        self._query = ViewFilter.Query(ret)
+        return self._query
+
+    def add_exclude(self, view):
+        self._query = None
+        self._excluded_views.add(view)
+
+    def remove_exclude(self, view):
+        try:
+            self._excluded_views.remove(view)
+            self._query = None
+            return True
+        except KeyError:
+            return False
+
+    def read_excludes(self):
+        yield from self._excluded_views
