Updated ushell [35954876]

#rnx

diff --git a/Engine/Extras/ushell/channels/unreal/core/pylib/cipher.py b/Engine/Extras/ushell/channels/unreal/core/pylib/cipher.py
--- a/Engine/Extras/ushell/channels/unreal/core/pylib/cipher.py
+++ b/Engine/Extras/ushell/channels/unreal/core/pylib/cipher.py
@@ -19,10 +19,10 @@ def _qr(x:List[int], a:int, b:int, c:int, d:int) -> None:
 
 #-------------------------------------------------------------------------------
 def _read_blocks(key:bytes, nonce:bytes, counter:int=0) -> Iterator[bytes]:
-    key = memoryview(key).cast("L")
-    nonce = memoryview(nonce).cast("L")
+    key = memoryview(key).cast("I")
+    nonce = memoryview(nonce).cast("I")
     state = [0x61707865, 0x3320646e, 0x79622d32, 0x6b206574, *key, 0, 0, *nonce]
-    out = memoryview(bytearray(4 * 16)).cast("L")
+    out = memoryview(bytearray(4 * 16)).cast("I")
     while True:
         state[12] = counter
         state[13] = counter >> 32
@@ -91,7 +91,7 @@ class Blob(object):
         for i in range(0, len(self._toc), TOC_ENTRY_SIZE):
             piece = self._toc[i:i + TOC_ENTRY_SIZE]
             piece = _chacha20(piece, key=key, nonce=nonce_toc)
-            piece = memoryview(piece).cast("L")
+            piece = memoryview(piece).cast("I")
             name = piece[4:].cast("B")
             if piece[2] == _hash(name):
                 name = name[:piece[3]]
@@ -101,7 +101,9 @@ class Blob(object):
                 data = self._data[offset:offset + size]
                 data = _chacha20(data, key=key, nonce=nonce)
 
-                yield bytes(name).decode(), data
+                name = bytes(name).replace(b"\\", b"/")
+                name = name.decode()
+                yield name, data
 
     def add(self, item:Path, key:memoryview, toc_name=None) -> None:
         size = item.stat().st_size
@@ -134,14 +136,16 @@ class Blob(object):
             data = _chacha20(inp.read(), key=key, nonce=nonce)
 
         # create a toc entry
-        name = (toc_name or str(item)).encode()
+        toc_name = (toc_name or str(item))
+        toc_name = toc_name.replace("\\", "/")
+        name = toc_name.encode()
         name_length = len(name)
         assert name_length <= 112
         name = name + secrets.token_bytes(112 - name_length)
 
         toc_entry = bytearray(TOC_ENTRY_SIZE)
         toc_entry[16:] = name
-        toc_entry = memoryview(toc_entry).cast("L")
+        toc_entry = memoryview(toc_entry).cast("I")
         toc_entry[0] = offset
         toc_entry[1] = size
         toc_entry[2] = _hash(name)
diff --git a/Engine/Extras/ushell/channels/unreal/perforce/cmds/cherrypick.py b/Engine/Extras/ushell/channels/unreal/perforce/cmds/cherrypick.py
--- a/Engine/Extras/ushell/channels/unreal/perforce/cmds/cherrypick.py
+++ b/Engine/Extras/ushell/channels/unreal/perforce/cmds/cherrypick.py
@@ -206,6 +206,10 @@ class Cherrypick(flow.cmd.Cmd):
         P4.change(i=True).run(input_data=cl_spec)
         dest_cl = P4.changes(c=info.clientName, m=1, s="pending").change
 
+        server_version = 0.0
+        if m := re.match(r".*\/(20\d+\.\d+)\/.*", info.serverVersion):
+            server_version = float(m.groups()[0])
+
         # Integrate or unshelve each input changelist
         specs = {}
         self.print_info("Cherrypicking into", dest_cl)
@@ -239,6 +243,11 @@ class Cherrypick(flow.cmd.Cmd):
                 "v" : self.args.virtual,
                 "c" : dest_cl,
             }
+
+            # 2024.1 no longer allows branch spec merges by default so we must also pass -F
+            if server_version >= 2024.1:
+                p4_args["F"] = True
+
             if desc.status == "pending":
                 p4_args["-bypass-exclusive-lock"] = True
                 print(f"Unshelving {cl} from", src_root, end="")
