Updated ushell [36023999]

#rnx

diff --git a/Engine/Extras/ushell/channels/flow/core/system/boot.py b/Engine/Extras/ushell/channels/flow/core/system/boot.py
--- a/Engine/Extras/ushell/channels/flow/core/system/boot.py
+++ b/Engine/Extras/ushell/channels/flow/core/system/boot.py
@@ -10,16 +10,16 @@ def _enable_vt100():
         return
 
     import ctypes
-    _win_dll = ctypes.LibraryLoader(ctypes.WinDLL)
-    _get_std_handle = _win_dll.kernel32.GetStdHandle
-    _get_console_mode = _win_dll.kernel32.GetConsoleMode
-    _set_console_mode = _win_dll.kernel32.SetConsoleMode
+    win_dll = ctypes.LibraryLoader(ctypes.WinDLL)
+    get_std_handle = win_dll.kernel32.GetStdHandle
+    get_console_mode = win_dll.kernel32.GetConsoleMode
+    set_console_mode = win_dll.kernel32.SetConsoleMode
 
     con_mode = ctypes.c_int()
-    stdout_handle = _get_std_handle(-11)
-    _get_console_mode(stdout_handle, ctypes.byref(con_mode))
+    stdout_handle = get_std_handle(-11)
+    get_console_mode(stdout_handle, ctypes.byref(con_mode))
     con_mode = con_mode.value | 4
-    _set_console_mode(stdout_handle, con_mode)
+    set_console_mode(stdout_handle, con_mode)
 
 #-------------------------------------------------------------------------------
 def _main():
diff --git a/Engine/Extras/ushell/channels/unreal/core/cmds/build.py b/Engine/Extras/ushell/channels/unreal/core/cmds/build.py
--- a/Engine/Extras/ushell/channels/unreal/core/cmds/build.py
+++ b/Engine/Extras/ushell/channels/unreal/core/cmds/build.py
@@ -82,7 +82,7 @@ class _PrettyPrinter(prettyprinter.Printer):
     def __init__(self, logger):
         self._source_re = re.compile(r"^(\s+\[[0-9/]+\]\s+|)([^ ]+\.[a-z]+)")
         self._error_re = re.compile("([Ee]rror:|ERROR|Exception: |error LNK|error C)")
-        self._warning_re = re.compile("(warning C|[Ww]arning: |note:)")
+        self._warning_re = re.compile("(warning [CV]|[Ww]arning: |note:)")
         self._progress_re = re.compile(r"^\s*@progress\s+('([^']+)'|\w+)(\s+(\d+))?")
         self._for_target_re = re.compile(r"^\*\* For ([^\s]+) \*\*")
         self._percent = 0
