Updated ushell [36145888]

#rnx

diff --git a/Engine/Extras/ushell/README.txt b/Engine/Extras/ushell/README.txt
--- a/Engine/Extras/ushell/README.txt
+++ b/Engine/Extras/ushell/README.txt
@@ -50,7 +50,7 @@ they can do is to list a few examples;
 11. .info
 
 Each of ushell's commands accept a '--help' argument which will show
-documentation for the command, details on how to its invoked, and descriptions
+documentation for the command, details on how it is invoked, and descriptions
 for the available options.
 
 # Tab Completion And Command History
@@ -58,18 +58,18 @@ for the available options.
 Anyone familiar with editing commands and recalling previous ones in Bash (i.e.
 Readline) will feel at home at a ushell prompt.
 
-There is extensive context-sensitive Tab completion available for commands and
+There is extensive context-sensitive tab completion available for commands and
 their arguments. Hitting Tab can help both to discover commands and their
-arguments, and also aid in fast convenient command entry. For example;
+arguments, and also aid in quick and convenient command entry. For example;
 
 1. .<tab><tab>            : displays available commands
 2. .b<tab>                : completes ".build "
 3. .run <tab><tab>        : shows options for the .run commands first argument
 4. .build editor --p<tab> : adds "--platform=" (further Tabs complete platforms)
 
-Ushell maintains a history of commands run which is persisted from one session to
-the next. Previous commands can be conveniently recalled in a few ways. To step
-backwards through prior commands by prefix use PgUp;
+Ushell maintains a history of previously run commands which persists from one
+session to the next. Previous commands can be conveniently recalled in a few
+ways. To step backwards through prior commands by prefix use PgUp;
 
 1. .bu<pgup>                : Cycle back through commands that started with ".bu"
 2. .run game switch<pgup>   : Iterate through previous runs on Switch.
diff --git a/Engine/Extras/ushell/channels/flow/posix/shells.py b/Engine/Extras/ushell/channels/flow/posix/shells.py
--- a/Engine/Extras/ushell/channels/flow/posix/shells.py
+++ b/Engine/Extras/ushell/channels/flow/posix/shells.py
@@ -156,8 +156,8 @@ class _Zsh(_Base):
 def _get_complete_zsh():
     return r"""
 _flow_prompt() {
-    FLOW_CWD=$(dirs)
     eval $(\$prompt --format=sh)
+    FLOW_CWD=$(dirs)
 }
 setopt promptsubst
 if [[ $FLOW_PROMPT ]]; then
@@ -169,7 +169,7 @@ fi
 
 function _flow_complete() {
     # setopt local_options xtrace
-    compadd $('$complete' -- ${^^words[@]:0:-1} $PREFIX...)
+    compadd - $('$complete' -- ${^^words[@]:0:-1} $PREFIX...)
 }
 compdef _flow_complete '$$$FLOW_CMDS$$$'
 """
diff --git a/Engine/Extras/ushell/channels/unreal/core/cmds/build.py b/Engine/Extras/ushell/channels/unreal/core/cmds/build.py
--- a/Engine/Extras/ushell/channels/unreal/core/cmds/build.py
+++ b/Engine/Extras/ushell/channels/unreal/core/cmds/build.py
@@ -188,12 +188,12 @@ class _PrettyPrinter(prettyprinter.Printer):
 #-------------------------------------------------------------------------------
 class _BuildCmd(unrealcmd.Cmd):
     """ The 'fileormod' argument allows the build to be constrained to a specific
-    module or source file to. It accepts one of the following forms as input;
+    module or source file. It accepts one of the following forms as input;
 
       CoreTest                 - builds the target's "CoreTest" module
       CoreTest/SourceFile.cpp  - compiles CoreTest/Private/**/SourceFile.cpp
       path/to/a/SourceFile.cpp - absolute or relative single file to compile
-      -                        - fuzzy find the file/modules to build (single dash)
+      .                        - fuzzy find the file/modules to build (a dot)
     """
     _variant  = unrealcmd.Arg("development", "The build variant to have the build build")
     fileormod = unrealcmd.Arg("", "Constrain build to a single file or module")
@@ -238,7 +238,7 @@ class _BuildCmd(unrealcmd.Cmd):
         if not file_or_module:
             return
 
-        if file_or_module == "-":
+        if file_or_module == ".":
             print("Select .cpp/.Build.cs...", end="")
             file_or_module = _fuzzy_find_code(".build [...] ")
             if not file_or_module:
diff --git a/Engine/Extras/ushell/channels/unreal/core/describe.flow.py b/Engine/Extras/ushell/channels/unreal/core/describe.flow.py
--- a/Engine/Extras/ushell/channels/unreal/core/describe.flow.py
+++ b/Engine/Extras/ushell/channels/unreal/core/describe.flow.py
@@ -54,7 +54,7 @@ ripgrep_darwin_arm = flow.describe.Tool()
 ripgrep_darwin_arm.version("14.1.0")
 ripgrep_darwin_arm.payload("[REDACTED_PROFILE_URL]/ripgrep/releases/download/$VERSION/ripgrep-$VERSION-aarch64-apple-darwin.tar.gz")
 ripgrep_darwin_arm.sha1("59d0d543a7e25831b29ebe2943da366d4c005d26")
-ripgrep_darwin_arm.platform("darwin-amd64")
+ripgrep_darwin_arm.platform("darwin-arm64")
 ripgrep_darwin_arm.bin("rg")
 
 
