Updated ushell [36429768]

#rnx

diff --git a/Engine/Extras/ushell/channels/unreal/core/cmds/stage.py b/Engine/Extras/ushell/channels/unreal/core/cmds/stage.py
--- a/Engine/Extras/ushell/channels/unreal/core/cmds/stage.py
+++ b/Engine/Extras/ushell/channels/unreal/core/cmds/stage.py
@@ -9,7 +9,7 @@ class _Impl(unrealcmd.Cmd):
     target   = unrealcmd.Arg(str, "Cooked target to stage")
     platform = unrealcmd.Arg("", "Platform whose cook should be staged")
     variant  = unrealcmd.Arg("development", "Build variant to stage binaries to copy")
-    style    = unrealcmd.Arg("auto", "Type of result to stage ('auto' makes the best choice for common use)")
+    style    = unrealcmd.Arg("auto", "Type of result to stage ('pak', 'nopak', 'zen', or 'auto' which is the default and makes the best choice for common use)")
     uatargs  = unrealcmd.Arg([str], "Additional arguments to pass to UAT")
 
     complete_target = ("game", "client", "server")
diff --git a/Engine/Extras/ushell/channels/unreal/core/prompt.py b/Engine/Extras/ushell/channels/unreal/core/prompt.py
--- a/Engine/Extras/ushell/channels/unreal/core/prompt.py
+++ b/Engine/Extras/ushell/channels/unreal/core/prompt.py
@@ -56,7 +56,7 @@ class Prompt(object):
             for line in file.readlines():
                 line = line.strip()
                 if line.startswith("ref: "):
-                    context.UE_PROJECT = line[5:].split("/")[-1]
+                    context.UE_PROJECT = "/".join(line[5:].split("/")[2:])
                 else:
                     context.UE_PROJECT = line[:6]
                 break
diff --git a/Engine/Extras/ushell/channels/unreal/core/pylib/unreal/_platform.py b/Engine/Extras/ushell/channels/unreal/core/pylib/unreal/_platform.py
--- a/Engine/Extras/ushell/channels/unreal/core/pylib/unreal/_platform.py
+++ b/Engine/Extras/ushell/channels/unreal/core/pylib/unreal/_platform.py
@@ -146,6 +146,12 @@ class Platform(object):
             if value: value = os.path.join(sdks_dir, value)
             yield Value(key, value)
 
+    @_optional_api
+    def get_stage_name(self, target):
+        if hasattr(self, "_get_stage_name"):
+            return self._get_stage_name(target)
+        return self.get_cook_form(target)
+
     @_optional_api
     def get_cook_form(self, target):
         form = self._get_cook_form(target)
diff --git a/Engine/Extras/ushell/channels/unreal/core/pylib/unreal/platforms/win64/android.py b/Engine/Extras/ushell/channels/unreal/core/pylib/unreal/platforms/win64/android.py
--- a/Engine/Extras/ushell/channels/unreal/core/pylib/unreal/platforms/win64/android.py
+++ b/Engine/Extras/ushell/channels/unreal/core/pylib/unreal/platforms/win64/android.py
@@ -106,6 +106,9 @@ class Platform(unreal.Platform):
     def _get_cook_flavor(self):
         return "ASTC"
 
+    def _get_stage_name(self, target):
+        return "Android"
+
     def _get_package_name(self):
         config = self.get_unreal_context().get_config()
         if ret := str(config.get("Engine", Platform.config_section, "PackageName")):
