Updated ushell [36467832]

#rnx

diff --git a/Engine/Extras/ushell/channels/unreal/core/pylib/unreal/_platform.py b/Engine/Extras/ushell/channels/unreal/core/pylib/unreal/_platform.py
--- a/Engine/Extras/ushell/channels/unreal/core/pylib/unreal/_platform.py
+++ b/Engine/Extras/ushell/channels/unreal/core/pylib/unreal/_platform.py
@@ -146,7 +146,6 @@ class Platform(object):
             if value: value = os.path.join(sdks_dir, value)
             yield Value(key, value)
 
-    @_optional_api
     def get_stage_name(self, target):
         if hasattr(self, "_get_stage_name"):
             return self._get_stage_name(target)
