Updated ushell [36504589]

#rnx

diff --git a/Engine/Extras/ushell/README.txt b/Engine/Extras/ushell/README.txt
--- a/Engine/Extras/ushell/README.txt
+++ b/Engine/Extras/ushell/README.txt
@@ -99,6 +99,11 @@ It is also possible to distribute a copy of ushell as a `ushell.zip` file
 alongside `ushell.ini`. This can be useful to deploy ushell to users regardless
 of which branch they may be working on.
 
+`.ushell gather` is available to gather ushell from `Engine/Extras/ushell/` and
+create a standalone deployment. This can be useful as a mechanism for managing
+site-specific commands, or having a different ushell update schedule compared to
+that for engine integrations.
+
 # Scripting
 
 There is modest support for scripting ushell with Batch scripts. This can be
diff --git a/Engine/Extras/ushell/channels/flow/core/cmds/debug.py b/Engine/Extras/ushell/channels/flow/core/cmds/debug.py
--- a/Engine/Extras/ushell/channels/flow/core/cmds/debug.py
+++ b/Engine/Extras/ushell/channels/flow/core/cmds/debug.py
@@ -7,6 +7,7 @@ import pprint
 import hashlib
 import marshal
 import flow.cmd
+import flow.describe
 from pathlib import Path
 from urllib.request import urlopen, URLError
 
@@ -83,6 +84,8 @@ class Wipe(flow.cmd.Cmd):
 
     def main(self):
         import shutil
+        from pathlib import Path
+
         channel = self.get_channel()
         system = channel.get_system()
         for item in (system.get_temp_dir(), system.get_working_dir()):
@@ -90,15 +93,19 @@ class Wipe(flow.cmd.Cmd):
             shutil.rmtree(item, ignore_errors=True)
             print("")
 
-        from pathlib import Path
+        del_dir = Path(system.get_working_dir()) / f"$cleaner/{os.getpid()}"
+        del_dir.mkdir(parents=True, exist_ok=True)
+
         tool_dir = Path(system.get_tools_dir())
         print("Tools dir:", tool_dir)
-        for item in tool_dir.glob("*"):
-            if item.is_dir():
-                print("Removing tool:", item.name, end="")
-                item = item.rename(str(item) + str(os.getpid()))
-                shutil.rmtree(item, ignore_errors=True)
-                print("...error" if item.is_dir() else "")
+        for i, item in enumerate(tool_dir.glob("*")):
+            if not item.is_dir():
+                continue
+
+            print("Removing tool:", item.name, end="")
+            item = item.rename(del_dir / str(i))
+            shutil.rmtree(item, ignore_errors=True)
+            print("...pending" if item.is_dir() else "")
 
 #-------------------------------------------------------------------------------
 class Paths(flow.cmd.Cmd):
@@ -125,14 +132,15 @@ class _Channels(flow.cmd.Cmd):
         channel = self.get_channel()
         system = channel.get_system()
 
-        self.channels = {}
-        channels_dir = system.get_working_dir() + "/channels/"
-        with os.scandir(channels_dir) as dirs:
-            for item in dirs:
-                with open(item.path + "/manifest", "rb") as data:
-                    manifest = marshal.load(data)
+        with open(system.get_working_dir() + "manifest", "rb") as x:
+            manifest = marshal.load(x)
 
-                self.channels[item.name] = manifest["tools"]
+        self.channels = {}
+        for item in (x for x in manifest["channels"] if x["tools"]):
+            self.channels[item["name"]] = {
+                "tools" : item["tools"],
+                "path" : item["path"],
+            }
 
 #-------------------------------------------------------------------------------
 class Tools(_Channels):
@@ -147,16 +155,31 @@ class Sha1s(_Channels):
     """ Shows channels' tools SHA1s """
 
     def _impl(self, name, manifest):
+        def import_script(script_path):
+            import importlib.util as import_util
+            spec = import_util.spec_from_file_location("", script_path)
+            module = import_util.module_from_spec(spec)
+            spec.loader.exec_module(module)
+            return module
+
+        tools = {}
+        module = import_script(Path(manifest["path"]) / "describe.flow.py")
+        for item_name in dir(module):
+            item = getattr(module, item_name)
+            if isinstance(item, flow.describe.Tool):
+                tools[item_name] = item
+
         print("--", name)
-        for tool_name, tool in manifest.items():
+        for tool_name, tool in tools.items():
             def on_data(file_name, recv_size, total_size, data):
                 if sys.stdout.isatty():
                     print("%.02f%%" % ((recv_size * 100) / total_size), end="\r")
                 digest.update(data)
 
             digest = hashlib.sha1()
-            for urls in tool["bundles"].values():
-                for url in urls:
+            for bundle_name, bundle in tool._bundles.items():
+                for url, file_type in bundle._payloads:
+                    url = url.replace("$VERSION", tool._version)
                     try: _http_get(url, on_data)
                     except URLError: pass
             sha1 = digest.hexdigest()
diff --git a/Engine/Extras/ushell/channels/flow/core/system/boot.py b/Engine/Extras/ushell/channels/flow/core/system/boot.py
--- a/Engine/Extras/ushell/channels/flow/core/system/boot.py
+++ b/Engine/Extras/ushell/channels/flow/core/system/boot.py
@@ -23,28 +23,15 @@ def _enable_vt100():
 
 #-------------------------------------------------------------------------------
 def _main():
-    # Flow is started like this; boot.py [working_dir] [channel_dir] -- [shell_args]
-    boot_dirs = []
-    shell_args = []
-    collector = boot_dirs
-    for arg in sys.argv[1:]:
-        if arg == "--":
-            collector = shell_args
-            continue
-
-        collector.append(arg)
-
+    # Get all the channels ready and aggregated
     lib_dir = Path(__file__).resolve().parent / "lib"
     sys.path.append(str(lib_dir))
-
-    # Get all the channels ready and aggregated
     import bootstrap
-    working_dir, *channels_dirs = (Path(x) for x in boot_dirs)
-    working_suffix = bootstrap.impl(working_dir, *channels_dirs)
+    inst_dir = bootstrap.impl()
 
     # Run the '$boot' command
     import run
-    ret = run.main((f"{working_dir}/{working_suffix}/manifest", "$boot", *shell_args))
+    ret = run.main((f"{inst_dir}/manifest", "$boot", *sys.argv[1:]))
     raise SystemExit(ret)
 
 
diff --git a/Engine/Extras/ushell/channels/flow/core/system/lib/bootstrap.py b/Engine/Extras/ushell/channels/flow/core/system/lib/bootstrap.py
--- a/Engine/Extras/ushell/channels/flow/core/system/lib/bootstrap.py
+++ b/Engine/Extras/ushell/channels/flow/core/system/lib/bootstrap.py
@@ -2,12 +2,14 @@
 
 import os
 import sys
+import shutil
+from pathlib import Path
+from typing import Iterable
+
 import fsutils
-import hashlib
 import marshal
 import flow.describe
 import subprocess as sp
-from pathlib import Path
 
 #-------------------------------------------------------------------------------
 class _Log(object):
@@ -134,6 +136,8 @@ def _validate_tool(name, tool):
 
 #-------------------------------------------------------------------------------
 def _acquire_tool(name, tool, target_dir, progress_cb):
+    import hashlib
+
     # Some boiler plate to provide a progress update
     payload_i = 0
     payload_n = sum(len(x._payloads) for x in tool._bundles.values())
@@ -223,7 +227,7 @@ def _manifest_tool(name, tool):
     }
 
 #-------------------------------------------------------------------------------
-def _install_tool(name, tool, channel_dir):
+def _install_tool(name, tool):
     manifest = _manifest_tool(name, tool)
 
     if req_platform := tool._platform:
@@ -283,112 +287,74 @@ def _install_tool(name, tool, channel_dir):
 
 
 #-------------------------------------------------------------------------------
-class _Action(object):
-    def __init__(self, name, input_path):
-        self._name = name
-        self._input_path = input_path
-
-    def needs_update(self):
-        return True
-
-
-
-#-------------------------------------------------------------------------------
-def _validate_command(channel_src_dir, name, command):
+def _validate_command(channel_dir, name, command):
     assert command._path, f"Command '{name}' missing call to invoke/composite()"
     assert command._py_path, f"Command '{name}' missing call to source()"
-    if not os.path.isfile(channel_src_dir + command._py_path):
+    if not (channel_dir / command._py_path).is_file():
         assert False, f"Command '{name}' missing source '{command._py_path}'"
 
 #-------------------------------------------------------------------------------
-def _select_update_type():
-    version_path = os.path.abspath(__file__ + "/../../version")
-    input_deps = (
-        version_path,             # primary version for schema/state invalidation
-        "../python/current/",     # embedded version of Python.
-    )
-    mtime = max(fsutils.get_mtime(x) for x in input_deps)
-    return _UpdateMaybe if fsutils.get_mtime("manifest") > mtime else _UpdateImpl
+def _build_channel(channel_name:str, channel_dir:Path) -> dict[str]:
+    _log.print(f"Updating channel '{channel_name}'")
+    _log.indent("Channel", channel_name)
+
+    # Load the describe.flow.py as a module.
+    channel_py = channel_dir / "describe.flow.py"
+    module = fsutils.import_script(channel_py)
+
+    def read_items_by_type(type):
+        for name, value in module.__dict__.items():
+            if isinstance(value, type):
+                yield name, value
+
+    # Get the channel's description.
+    name, channel = next(read_items_by_type(flow.describe.Channel), (None, None))
+    assert channel, f"Channel() instance missing from channel '{channel_name}'"
+    assert channel._version != None, f"Channel '{channel_name}' has no version"
+
+    # Install channel's pips.
+    if channel._pips:
+        _log.print("Pips:")
+        _log.indent("Pips")
+        for pip_name in channel._pips:
+            _log.print_write(pip_name)
+            cmd = (sys.executable, "-Xutf8", "-Esum", "pip", "install", pip_name)
+            result = sp.run(cmd, stdout=sp.DEVNULL, stderr=sp.DEVNULL)
+            if result.returncode:
+                _log.print(" ...failed");
+        _log.unindent()
 
-#-------------------------------------------------------------------------------
-class _UpdateImpl(_Action):
-    def run(self):
-        _log.print(f"Updating channel '{self._name}'")
-        _log.indent("Channel", self._name)
-
-        # Clean up the output. Updates are an install from scratch.
-        out_dir = f"channels/{self._name}/"
-        fsutils.rmdir(out_dir)
-        out_dir = fsutils.WorkPath(out_dir)
-
-        # Load the describe.flow.py as a module.
-        channel_py = self._input_path + "describe.flow.py"
-        module = fsutils.import_script(channel_py)
-
-        def read_items_by_type(type):
-            for name, value in module.__dict__.items():
-                if isinstance(value, type):
-                    yield name, value
-
-        # Get the channel's description.
-        name, channel = next(read_items_by_type(flow.describe.Channel), (None, None))
-        assert channel, f"Channel() instance missing from channel '{self._name}'"
-        assert channel._version != None, f"Channel '{self._name}' has no version"
-
-        # Install channel's pips.
-        if channel._pips:
-            _log.print("Pips:")
-            _log.indent("Pips")
-            for pip_name in channel._pips:
-                _log.print_write(pip_name)
-                cmd = (sys.executable, "-Xutf8", "-Esum", "pip", "install", pip_name)
-                result = sp.run(cmd, stdout=sp.DEVNULL, stderr=sp.DEVNULL)
-                if result.returncode:
-                    _log.print(" ...failed");
-            _log.unindent()
-
-        # Tools
-        tools_manifest = {}
-        tools = [x for x in read_items_by_type(flow.describe.Tool)]
-        if tools:
-            _log.print("Tools:")
-            _log.indent("Tools")
-            for name, tool in tools:
-                _validate_tool(name, tool)
-
-            for name, tool in tools:
-                tools_manifest[name] = _install_tool(name, tool, out_dir)
-            _log.unindent()
-
-        # Commands
-        commands = [x for x in read_items_by_type(flow.describe.Command)]
-        for name, command in commands:
-            _validate_command(self._input_path, name, command)
-            command._path = command._path or (name)
-            command._path = (command._prefix, *command._path)
-
-        # Build manifest
-        manifest = {}
-        manifest["name"] = self._name
-        manifest["path"] = self._input_path
-        manifest["parent"] = channel._parent or "flow.core"
-        manifest["tools"] = tools_manifest
-        manifest["commands"] = tuple(v.__dict__ for k,v in commands)
-
-        # Success. Write the action's output
-        with open(out_dir + "manifest", "wb") as x:
-            marshal.dump(manifest, x)
-
-        out_dir.keep()
-        _log.print("Done\n")
+    # Tools
+    tools_manifest = {}
+    tools = [x for x in read_items_by_type(flow.describe.Tool)]
+    if tools:
+        _log.print("Tools:")
+        _log.indent("Tools")
+        for name, tool in tools:
+            _validate_tool(name, tool)
+
+        for name, tool in tools:
+            tools_manifest[name] = _install_tool(name, tool)
         _log.unindent()
 
-#-------------------------------------------------------------------------------
-class _UpdateMaybe(_UpdateImpl):
-    def needs_update(self):
-        in_mtime = fsutils.get_mtime(self._input_path + "describe.flow.py")
-        out_mtime = fsutils.get_mtime(f"channels/{self._name}/manifest")
-        return out_mtime <= in_mtime
+    # Commands
+    commands = [x for x in read_items_by_type(flow.describe.Command)]
+    for name, command in commands:
+        _validate_command(channel_dir, name, command)
+        command._path = command._path or (name)
+        command._path = (command._prefix, *command._path)
+
+    # Build manifest
+    manifest = {}
+    manifest["name"] = channel_name
+    manifest["path"] = str(channel_dir) + "/"
+    manifest["parent"] = channel._parent or "flow.core"
+    manifest["tools"] = tools_manifest
+    manifest["commands"] = tuple(v.__dict__ for k,v in commands)
+
+    _log.print("Done\n")
+    _log.unindent()
+    return manifest
 
 
 
@@ -439,9 +405,13 @@ def _create_shims(manifest):
     builder = shims.Builder(shims_dir)
     builder.clean()
 
+    manifest_path = Path().resolve()
+    manifest_path = manifest_path.parent / manifest_path.name.split("$")[0]
+    manifest_path /= "manifest"
+    manifest_path = str(manifest_path)
+
     # Create command shims
     run_py = os.path.abspath(__file__ + "/../../run.py")
-    manifest_path = os.path.abspath("manifest")
     args_prefix = ("-Xutf8", "-Esu", run_py, manifest_path)
     for name in (x for x in manifest["cmd_tree"] if isinstance(x, str)):
         builder.create_python_shim(name, *args_prefix, name)
@@ -462,148 +432,234 @@ def _create_shims(manifest):
     builder.commit()
 
 #-------------------------------------------------------------------------------
-class _Finalise(_Action):
-    def needs_update(self):
-        self._manifest_paths = [x.path + "/manifest" for x in fsutils.read_dirs(self._input_path)]
-        mtime = max(fsutils.get_mtime(x) for x in self._manifest_paths)
-        mtime = max(mtime, fsutils.get_mtime(self._input_path))
-        return fsutils.get_mtime("manifest") <= mtime
-
-    def run(self):
-        # Load channel manifests
-        manifests = []
-        for manifest_path in self._manifest_paths:
-            with open(manifest_path, "rb") as x:
-                manifest = marshal.load(x)
-                manifests.append(manifest)
-
-        # Assign each manifest an index
-        for i in range(len(manifests)):
-            manifests[i]["index"] = i
-
-        # Connect children to their parents
-        manifests = {x["name"]:x for x in manifests}
-        for manifest in manifests.values():
-            parent = manifest["parent"]
-            assert parent in manifests, f"Invalid parent '{parent}' for channel '{manifest['name']}'"
-            manifest["parent"] = manifests[parent] if parent != manifest["name"] else None
-
-        # Sort the manifests topologically
-        topo_manifests = {}
-
-        def topo_insert(manifest):
-            if manifest:
-                topo_insert(manifest["parent"])
-                topo_manifests[manifest["name"]] = manifest
-
-        for manifest in manifests.values():
-            topo_insert(manifest)
-
-        manifests = topo_manifests
-        _log.write("Channels;", *(x["name"] for x in manifests.values()))
-
-        # Collect various paths for each channel
-        pylib_dirs = []
-        for manifest in manifests.values():
-            pylib_dir = manifest["path"] + "pylib/"
-            if os.path.isdir(pylib_dir):
-                pylib_dirs.append(pylib_dir)
-        _log.write("Pylib dirs;", *(x for x in pylib_dirs))
-
-        # Build the primary manifest.
-        primary = {
-            "channels"   : _carve_channels(manifests.values()),
-            "pylib_dirs" : tuple(pylib_dirs),
-            "tools_dir"  : os.path.abspath("../tools") + "/",
-            "cmd_tree"   : _plant_cmd_tree(manifests.values()),
-        }
+def _finalise(manifests):
+    # Assign each manifest an index
+    for i in range(len(manifests)):
+        manifests[i]["index"] = i
+
+    # Connect children to their parents
+    manifests = {x["name"]:x for x in manifests}
+    for manifest in manifests.values():
+        parent = manifest["parent"]
+        assert parent in manifests, f"Invalid parent '{parent}' for channel '{manifest['name']}'"
+        manifest["parent"] = manifests[parent] if parent != manifest["name"] else None
+
+    # Sort the manifests topologically
+    topo_manifests = {}
+
+    def topo_insert(manifest):
+        if manifest:
+            topo_insert(manifest["parent"])
+            topo_manifests[manifest["name"]] = manifest
+
+    for manifest in manifests.values():
+        topo_insert(manifest)
+
+    manifests = topo_manifests
+    _log.write("Channels;", *(x["name"] for x in manifests.values()))
+
+    # Collect various paths for each channel
+    pylib_dirs = []
+    for manifest in manifests.values():
+        pylib_dir = manifest["path"] + "pylib/"
+        if os.path.isdir(pylib_dir):
+            pylib_dirs.append(pylib_dir)
+    _log.write("Pylib dirs;", *(x for x in pylib_dirs))
+
+    # Build the primary manifest.
+    primary = {
+        "channels"   : _carve_channels(manifests.values()),
+        "pylib_dirs" : tuple(pylib_dirs),
+        "tools_dir"  : os.path.abspath("../tools") + "/",
+        "cmd_tree"   : _plant_cmd_tree(manifests.values()),
+    }
 
-        # Create shims
-        _create_shims(primary)
+    # Create shims
+    _create_shims(primary)
 
-        # Write the primary manifest out.
-        with open("manifest", "wb") as x:
-            marshal.dump(primary, x)
+    # Write the primary manifest out.
+    with open("manifest", "wb") as x:
+        marshal.dump(primary, x)
 
 
 
 #-------------------------------------------------------------------------------
-class _Prune(_Action):
-    def run(self):
-        _log.print_write(f"Removing old channel '{self._name}'")
-        fsutils.rmdir(self._input_path)
+class _Cleaner(object):
+    def __init__(self, working_dir:Path) -> None:
+        self._dir = working_dir / f"$cleaner/{os.getpid()}"
+        self._dir.mkdir(parents=True, exist_ok=True)
+        self._count = 0
+
+    def __del__(self) -> None:
+        shutil.rmtree(self._dir.parent, ignore_errors=True)
+
+    def delete(self, path:Path) -> bool:
+        try: path.rename(self._dir / str(self._count))
+        except: pass
+        self._count += 1
 
 
 
 #-------------------------------------------------------------------------------
-def impl(working_dir:Path, *channels_dirs:Path):
-    channels_dirs = [str(x.resolve()) for x in channels_dirs if x.is_dir()]
-    working_dir = str(working_dir.resolve())
-
-    def _read_channels(dir, depth=0):
-        for item in fsutils.read_dirs(dir):
-            if os.path.isfile(item.path + "/describe.flow.py"):
-                yield item.name, os.path.abspath(item.path) + "/"
-            elif depth == 0:
-                for name, path in _read_channels(item.path, depth + 1):
-                    yield f"{item.name}.{name}", path
-
-    # Seed the working dir with a hash of our location so we can detect movement
-    loc_hash = 5381
-    for c in os.path.abspath(__file__):
-        loc_hash = ((loc_hash << 5) + loc_hash) + (ord(c) | 0x20)
-        loc_hash &= 0x7fffffff
-    working_suffix = "flow_%08x" % loc_hash
-
-    # Change to the main working dir
-    working_dir += f"/{working_suffix}/"
-    try: os.makedirs(working_dir)
-    except FileExistsError: pass
-    boot_dir = os.getcwd()
-    os.chdir(working_dir)
-
-    lock = fsutils.DirLock("..")
-
-    # Build a list of actions to do. The flow channel must be first
-    actions = {"flow.core" : None}
-
-    # First; add existing channels for removal
-    for dest_dir in fsutils.read_dirs("channels"):
-        actions[dest_dir.name] = _Prune(dest_dir.name, dest_dir.path + "/")
-
-    # Second; add channels from source to get updated.
-    update_type = _select_update_type()
-    for channels_dir in channels_dirs:
-        for name, path in _read_channels(channels_dir):
-            actions[name] = update_type(name, path)
-
-    # Lastly; aggregate channels into a manifest and create command shims
-    actions["*finalise"] = _Finalise("*finalise", "channels/")
-
-    # If there's something to do, do it or wait for it to be done
-    for action in (x for x in actions.values() if x.needs_update()):
-        if lock.claim():
-            global _log
-            if not _log:
-                _log = _Log(working_dir + "setup.log")
-                _log.write("Working dir;", working_dir)
-
-            try:
-                _log.write("-- Running action;", type(action).__name__)
-                action.run()
-            except:
-                import traceback
-                _log.write(traceback.format_exc())
-                raise
-            continue
+def _hashen(x:str) -> int:
+    ret = 5381
+    for c in x:
+        ret = ((ret << 5) + ret) + ord(c)
+        ret &= 0x7fffffff
+    return ret
 
-        import time
-        print("Waiting for another instance to provision.", end="")
-        while not lock.can_claim():
-            time.sleep(1)
-            print(".", end="")
 
-        break
 
-    os.chdir(boot_dir)
-    return working_suffix
+#-------------------------------------------------------------------------------
+class _StateDir(object):
+    _dir:Path
+
+    def __init__(self, root:Path):
+        loc_hash = _hashen(str(Path(__file__).resolve()))
+        self._dir = root / ("flow_%08x" % loc_hash)
+
+    def get_dir(self) -> Path:
+        return self._dir
+
+    def get_age(self) -> int:
+        try: return (self._dir / "manifest").stat().st_mtime
+        except FileNotFoundError: return 0
+
+
+
+#-------------------------------------------------------------------------------
+class _Sources(object):
+    _sources    :dict[str, str]
+    _key        :int
+
+    def __init__(self) -> None:
+        self._sources = {}
+        self._key = 0
+
+    def __iter__(self) -> Iterable[[str, Path]]:
+        yield from self._sources.items()
+
+    def add_source(self, name:str, flow_py_path:Path) -> None:
+        self._sources[name] = flow_py_path
+        self._key ^= _hashen(name);
+
+    def get_key(self) -> str:
+        return f"sources_{self._key:08x}"
+
+    def get_age(self) -> float:
+        deps = (
+            Path(sys.executable),
+            Path(__file__).parent.parent / "version",
+            *(self._sources.values()),
+        )
+        return max(x.stat().st_mtime for x in deps)
+
+
+
+#-------------------------------------------------------------------------------
+def _get_sources() -> _Sources:
+    sources = _Sources()
+
+    def _add_source(sources:_Sources, source_dir:Path, group_name:str="") -> int:
+        count = 0
+        for item in source_dir.glob("*"):
+            if not item.is_dir():
+                continue
+
+            candidate = item / "describe.flow.py"
+            if candidate.is_file():
+                sources.add_source(group_name + item.name, candidate)
+                count += 1
+            elif not group_name:
+                count += _add_source(sources, item, item.name + ".")
+
+        return count
+
+    for item in Path(__file__).parents:
+        if item.name == "channels":
+            _add_source(sources, item)
+            break
+
+    from_home = Path.home() / ".ushell/channels"
+    if from_home.is_dir():
+        _add_source(sources, from_home, "$home.")
+
+    if from_env := os.getenv("FLOW_CHANNELS_DIR"):
+        import shlex
+        for item in shlex.split(from_env):
+            item = Path(item).resolve()
+            if item.is_dir():
+                _add_source(sources, item, "$env.")
+
+    return sources
+
+#-------------------------------------------------------------------------------
+def _update(working_dir:Path, sources:_Sources) -> None:
+    class _CwdScope(object):
+        def __init__(self, dir:Path): self._prev_dir = os.getcwd(); os.chdir(dir)
+        def __del__(self) -> None:    os.chdir(self._prev_dir)
+    cwd_scope = _CwdScope(working_dir)
+
+    # Fire up our log. This needs to be deleted on the way out to not lock dirs
+    global _log
+    _log = _Log(working_dir / "setup.log")
+    _log.write("Working dir;", working_dir)
+    class _LogScope(object):
+        def __del__(self): global _log; _log = None
+    log_scope = _LogScope()
+
+    try:
+        # Build all the channels
+        manifests = []
+        for name, flow_py_path in sources:
+            manifest = _build_channel(name, flow_py_path.parent)
+            manifests.append(manifest)
+
+        # Aggregate all channels into a manifest and create command shims
+        _finalise(manifests)
+
+        # Add the marker to detect new/removed channels
+        with (working_dir / sources.get_key()).open("wb"):
+            pass
+    except:
+        import traceback
+        _log.write(traceback.format_exc())
+        raise
+    finally:
+        del cwd_scope
+        del log_scope
+
+#-------------------------------------------------------------------------------
+def impl() -> Path:
+    # Assume we've our own Python (explicit or a venv) and use that to determine
+    # where the working directory is.
+    py_bin_path = Path(sys.executable)
+    for item in py_bin_path.parents:
+        if item.parent.name == "python":
+            working_dir = item.parent.parent
+            break
+    else:
+        working_dir = Path.home() / ".ushell"
+
+    sources = _get_sources()
+
+    state_dir = _StateDir(working_dir)
+    if sources.get_age() <= state_dir.get_age():
+        if (state_dir.get_dir() / sources.get_key()).is_file():
+            return state_dir.get_dir()
+
+    dest_dir = state_dir.get_dir()
+
+    cleaner = _Cleaner(working_dir)
+    cleaner.delete(dest_dir)
+
+    temp_dir = dest_dir.parent / (dest_dir.name + f"${os.getpid()}")
+    try:
+        temp_dir.mkdir(parents=True)
+        _update(temp_dir, sources)
+        dest_dir = temp_dir.rename(dest_dir)
+    except:
+        cleaner.delete(temp_dir)
+        raise
+
+    return dest_dir
diff --git a/Engine/Extras/ushell/channels/flow/nt/boot.bat b/Engine/Extras/ushell/channels/flow/nt/boot.bat
--- a/Engine/Extras/ushell/channels/flow/nt/boot.bat
+++ b/Engine/Extras/ushell/channels/flow/nt/boot.bat
@@ -7,7 +7,6 @@ chcp 65001 1>nul 2>nul
 :: Set up the working directory flow can store its necessities in, and the dir
 :: that contains the channels that implement the shell.
 set _working="%localappdata%\ushell\.working"
-set _channels="%~dp0..\..\..\channels"
 if defined flow_working_dir (
     set _working="%flow_working_dir%"
 )
@@ -17,12 +16,6 @@ if exist "%userprofile%\.ushell\hooks\boot.bat" (
     call "%userprofile%\.ushell\hooks\boot.bat"
 )
 
-:: Make it easier for users to add their own channels
-set _channels=%_channels% "%userprofile%\.ushell\channels"
-if defined flow_channels_dir (
-    set _channels=%_channels% %flow_channels_dir%
-)
-
 :: Provision Python. We shim through a cmd.exe for simpler error handling in
 :: provision.bat (see the script's use of call:on_error).
 cmd.exe /d/c ""%~dp0\provision.bat" %_working%"
@@ -32,7 +25,7 @@ if errorlevel 1 exit /b 1
 set _cookie=%time::=%
 set _cookie=%_cookie:~1%
 set _cookie=%temp%\ushell\cmd_boot_%_cookie%
-%_working%\python\current\flow_python.exe -Xutf8 -Esu "%~dp0..\core\system\boot.py" %_working% %_channels% -- "--bootarg=cmd,%_cookie%" %*
+%_working%\python\current\flow_python.exe -Xutf8 -Esu "%~dp0..\core\system\boot.py" "--bootarg=cmd,%_cookie%" %*
 
 :: If --help was given then we return the user to where they came from. The "127"
 :: comes from _flick
diff --git a/Engine/Extras/ushell/channels/flow/nt/boot.ps1 b/Engine/Extras/ushell/channels/flow/nt/boot.ps1
--- a/Engine/Extras/ushell/channels/flow/nt/boot.ps1
+++ b/Engine/Extras/ushell/channels/flow/nt/boot.ps1
@@ -7,12 +7,6 @@ if(![String]::IsNullOrEmpty($env:flow_working_dir)) {
 
 Write-Debug "Working dir: $Working"
 
-$Channels = @(Join-Path $PSScriptRoot "..\..\..\channels")
-$Channels += Join-Path $env:userprofile "\.ushell\channels"
-$Channels += $env:flow_channels_dir
-
-Write-Debug "Channels: $Channels"
-
 $ProvisionBat = Join-Path $PSScriptRoot "provision.bat"
 Write-Debug "Provisioning: $ProvisionBat"
 & "cmd.exe" "/d/c" "`"$ProvisionBat`" $Working"
@@ -32,7 +26,7 @@ $Cookie = Join-Path $TempDir "cmd_boot_$(New-Guid).ps1"
 $Cookie = New-Item $Cookie
 Write-Debug "Cookie: $Cookie"
 
-& $PythonPath "-Xutf8" "-Esu" $Bootpy $Working @Channels "--" "--bootarg=pwsh,$Cookie" $Args | Out-Host
+& $PythonPath "-Xutf8" "-Esu" $Bootpy "--bootarg=pwsh,$Cookie" $Args | Out-Host
 $PythonReturn = $?
 
 if(!$PythonReturn) {
diff --git a/Engine/Extras/ushell/channels/flow/nt/boot.sh b/Engine/Extras/ushell/channels/flow/nt/boot.sh
--- a/Engine/Extras/ushell/channels/flow/nt/boot.sh
+++ b/Engine/Extras/ushell/channels/flow/nt/boot.sh
@@ -11,12 +11,7 @@ if [ ! $? ]; then
     exit 1
 fi
 
-declare -a channels_dir
-channels_dir+=("$self_dir/../..")
-channels_dir+=("$USERPROFILE/.ushell/channels")
-channels_dir+=("$FLOW_CHANNELS_DIR")
-
-"$working_dir/python/current/flow_python.exe" -Esu "$self_dir/../core/system/boot.py" "$working_dir" "${channels_dir[@]}" -- $@
+"$working_dir/python/current/flow_python.exe" -Esu "$self_dir/../core/system/boot.py" $@
 if [ ! $? ]; then
     echo boot.py failed
     exit 1
diff --git a/Engine/Extras/ushell/channels/flow/posix/boot.sh b/Engine/Extras/ushell/channels/flow/posix/boot.sh
--- a/Engine/Extras/ushell/channels/flow/posix/boot.sh
+++ b/Engine/Extras/ushell/channels/flow/posix/boot.sh
@@ -11,14 +11,7 @@ if ! $self_dir/provision.sh $working_dir; then
     exit 1
 fi
 
-declare -a channels_dir
-channels_dir+=("$self_dir/../..")
-channels_dir+=("$HOME/.ushell/channels")
-if [[ -n "$FLOW_CHANNELS_DIR" ]]; then
-    channels_dir+=("$FLOW_CHANNELS_DIR")
-fi
-
-"$working_dir/python/current/bin/python" -Esu "$self_dir/../core/system/boot.py" "$working_dir" "${channels_dir[@]}" -- "$@"
+"$working_dir/python/current/bin/python" -Esu "$self_dir/../core/system/boot.py" "$@"
 if [ ! $? ]; then
     echo boot.py failed
     exit 1
diff --git a/Engine/Extras/ushell/channels/unreal/core/cmds/gather.py b/Engine/Extras/ushell/channels/unreal/core/cmds/gather.py
new file mode 100644
--- /dev/null
+++ b/Engine/Extras/ushell/channels/unreal/core/cmds/gather.py
@@ -0,0 +1,82 @@
+import shutil
+import flow.cmd
+from pathlib import Path
+
+#-------------------------------------------------------------------------------
+def _find_ushell_dir(start_path):
+    start_path = Path(start_path) or Path(__file__)
+    for item in start_path.parents:
+        if (item / ("ushell.bat")).is_file():
+            return item;
+
+#-------------------------------------------------------------------------------
+class Gather(flow.cmd.Cmd):
+    """ Gathers an Engine-deployed version of ushell into a standalone version. """
+    destdir     = flow.cmd.Arg(Path, "Destination directory to gather into")
+    srcdir      = flow.cmd.Arg("", "Source to gather (default is this ushell instance)")
+    overwrite   = flow.cmd.Opt(False, "Overwrite destination if it exists")
+
+    def main(self):
+        self.print_info("Src/dest paths")
+        src_dir = _find_ushell_dir(self.args.srcdir)
+        if not src_dir:
+            raise ValueError("Unable to locate '<src_dir>/ushell.bat'")
+        print("src:", src_dir)
+
+        # Validate source
+        for item in src_dir.parents:
+            candidate = item / "GenerateProjectFiles.bat"
+            print("    ", candidate, end="\r")
+            if candidate.is_file():
+                print("  !")
+                break
+            print("  -")
+        else:
+            raise ValueError(f"'{src_dir}' does not appear to be located in a UE branch")
+
+        # Validate destination
+        print("dest:", self.args.destdir)
+        dest_dir = self.args.destdir
+        if dest_dir.exists():
+            if not self.args.overwrite:
+                raise ValueError(f"'{dest_dir}' already exists")
+            else:
+                if not (dest_dir / "ushell.bat").is_file():
+                    raise ValueError("'{dest_dir}/ushell.bat' expected - not overwriting")
+                shutil.rmtree(dest_dir)
+
+        dest_dir.mkdir(parents=True, exist_ok=True)
+
+        # Build copy rota
+        self.print_info("Collecting source")
+
+        def is_src_file(item):
+            if not item.is_file():          return
+            if item.suffix == ".pyc":       return
+            if item.parent.name == "tps":   return
+            return item
+
+        print(src_dir / "**")
+        rota = [(x, x.parent.relative_to(src_dir)) for x in src_dir.rglob("*") if is_src_file(x)]
+
+        p_dest = Path("channels/unreal/core/pylib/unreal/platforms/win64")
+        for p_dir in (src_dir.parent.parent / "Platforms").glob("*"):
+            p_dir /= "Extras/ushell"
+            if not p_dir.is_dir():
+                continue
+            print(p_dir / "*")
+            rota += [(x, p_dest) for x in p_dir.glob("*") if is_src_file(x)]
+
+        print("Source files:", len(rota))
+
+        # Do the copy
+        self.print_info("Gathering")
+        for i, (item, item_dest_dir) in enumerate(rota):
+            dest_suffix = item_dest_dir / item.name
+            dest_item = dest_dir / dest_suffix
+            if dest_item.name.startswith("platform_"):
+                dest_item = dest_item.parent / dest_item.name[9:]
+            print("%3d" % i, dest_suffix)
+            dest_item.parent.mkdir(parents=True, exist_ok=True)
+            shutil.copyfile(item, dest_item)
+        print("... done!")
diff --git a/Engine/Extras/ushell/channels/unreal/core/describe.flow.py b/Engine/Extras/ushell/channels/unreal/core/describe.flow.py
--- a/Engine/Extras/ushell/channels/unreal/core/describe.flow.py
+++ b/Engine/Extras/ushell/channels/unreal/core/describe.flow.py
@@ -255,6 +255,11 @@ odsc_all = flow.describe.Command()
 odsc_all.source("cmds/odsc.py", "All")
 odsc_all.invoke("cook", "odsc", "all")
 
+#-------------------------------------------------------------------------------
+gather = flow.describe.Command()
+gather.source("cmds/gather.py", "Gather")
+gather.invoke("ushell", "gather")
+
 #-------------------------------------------------------------------------------
 prompt = flow.describe.Command()
 prompt.source("prompt.py", "Prompt")
diff --git a/Engine/Extras/ushell/channels/unreal/perforce/cmds/workspace.py b/Engine/Extras/ushell/channels/unreal/perforce/cmds/workspace.py
--- a/Engine/Extras/ushell/channels/unreal/perforce/cmds/workspace.py
+++ b/Engine/Extras/ushell/channels/unreal/perforce/cmds/workspace.py
@@ -163,7 +163,7 @@ class Workspace(flow.cmd.Cmd):
 
         if not self.args.dryrun:
             print("Creating...", end="")
-            P4.client(i=True).run(input_data=client_spec)
+            P4(u=username).client(i=True).run(input_data=client_spec)
             print(" done")
 
         # Write a P4CONFIG file
diff --git a/Engine/Extras/ushell/channels/unreal/perforce/pylib/p4utils/syncer.py b/Engine/Extras/ushell/channels/unreal/perforce/pylib/p4utils/syncer.py
--- a/Engine/Extras/ushell/channels/unreal/perforce/pylib/p4utils/syncer.py
+++ b/Engine/Extras/ushell/channels/unreal/perforce/pylib/p4utils/syncer.py
@@ -1,5 +1,6 @@
 import re
 import time
+import flow.cmd
 import threading
 from peafour import P4
 from . import view
