Add UnrealToolbox config file for ushell.

diff --git a/Engine/Extras/ushell/Toolbox.json b/Engine/Extras/ushell/Toolbox.json
new file mode 100644
--- /dev/null
+++ b/Engine/Extras/ushell/Toolbox.json
@@ -0,0 +1,10 @@
+{
+	"statusPanelLinks": [
+		{
+			"label": "ushell",
+			"fileName": "$(COMSPEC)",
+			"arguments": "/D /K \"\"$(ToolDir)\\ushell.bat\" --project=\"$(ProjectFile)\"\"",
+			"workingDir": "$(ProjectDir)"
+		}
+	]
+}
