Updated ushell [36664915]

#rnx

diff --git a/Engine/Extras/ushell/channels/unreal/core/cmds/uat.py b/Engine/Extras/ushell/channels/unreal/core/cmds/uat.py
--- a/Engine/Extras/ushell/channels/unreal/core/cmds/uat.py
+++ b/Engine/Extras/ushell/channels/unreal/core/cmds/uat.py
@@ -1,6 +1,7 @@
 # Copyright Epic Games, Inc. All Rights Reserved.
 
 import os
+import unreal
 import unrealcmd
 
 #-------------------------------------------------------------------------------
@@ -64,8 +65,9 @@ class Uat(unrealcmd.MultiPlatformCmd):
         if project := ue_context.get_project():
             if args[0] and not self.args.unprojected:
                 args = (*args, "-project=" + str(project.get_path()))
-            if not self.args.allscripts:
-                args = (*args, "-ScriptsForProject=" + project.get_name())
+            if ue_context.get_type() != unreal.ContextType.FOREIGN:
+                if not self.args.allscripts:
+                    args = (*args, "-ScriptsForProject=" + project.get_name())
 
         exec_context = self.get_exec_context()
         cmd = exec_context.create_runnable(cmd, *args, *self.args.uatargs)
diff --git a/Engine/Extras/ushell/channels/unreal/perforce/cmds/gui.py b/Engine/Extras/ushell/channels/unreal/perforce/cmds/gui.py
--- a/Engine/Extras/ushell/channels/unreal/perforce/cmds/gui.py
+++ b/Engine/Extras/ushell/channels/unreal/perforce/cmds/gui.py
@@ -10,15 +10,19 @@ from pathlib import Path
 def _check_for_existing_nt(client):
     client = client.encode()
 
-    wmic = ("wmic.exe", "process", "get", "name,commandline")
+    wmic = ("tasklist.exe", "/v", "/nh", "/fo:csv", "/fi", "imagename eq p4v.exe")
     proc = subprocess.Popen(wmic, stdout=subprocess.PIPE)
     for line in iter(proc.stdout.readline, b""):
         if b"p4v.exe" in line and client in line:
-            return True
+            ret = True
+            break
+    else:
+        ret = False
     proc.stdout.close()
     proc.wait()
+    return ret
 
-_check_for_existing = _check_for_existing_nt if os.name == "nt" else lambda x: ""
+_check_for_existing = _check_for_existing_nt if os.name == "nt" else lambda x: False
 
 
 
