Updated ushell [36785141]

#rnx

diff --git a/Engine/Extras/ushell/channels/flow/core/system/flow/_flick.py b/Engine/Extras/ushell/channels/flow/core/system/flow/_flick.py
--- a/Engine/Extras/ushell/channels/flow/core/system/flow/_flick.py
+++ b/Engine/Extras/ushell/channels/flow/core/system/flow/_flick.py
@@ -171,7 +171,7 @@ class Cmd(object):
             if opt.needs_param():
                 value = opt.get_value()
                 if value is None:
-                    opt_name += " " + opt.get_type().__name__.upper()
+                    opt_name += "=" + opt.get_type().__name__.upper()
                 else:
                     opt_name += f"[={opt.get_type().__name__.upper()}]"
             opts.append(("--" + opt_name, opt.get_desc()))
@@ -180,17 +180,17 @@ class Cmd(object):
         # Output
         col_a = 0 if not args else max(len(x[0]) for x in args)
         col_o = max(len(x[0]) for x in opts)
-        col_0 = max(col_a, col_o)
+        col_0 = max(col_a, col_o) + 6
 
         import textwrap
-        width = max(40, 80 - col_0 - 4) # 3 = leader + separator = 2 + 2
+        width = max(40, 80 - col_0) # n = leader + dots = 2 + 4
         def print_param(name, desc):
-            print("  %-*s  " % (col_0, name), end="")
+            print(" ", name, "." * (col_0 - 4 - len(name)), end=" ")
             leader = ""
-            wrapper = textwrap.TextWrapper(width=width, subsequent_indent="  ")
+            wrapper = textwrap.TextWrapper(width=width)
             for line in wrapper.wrap(desc):
                 print(leader, line, sep="")
-                leader = leader or (" " * (col_0 + 4))
+                leader = leader or (" " * col_0)
 
         if args:
             print("\nARGS:")
diff --git a/Engine/Extras/ushell/channels/unreal/core/cmds/ddc.py b/Engine/Extras/ushell/channels/unreal/core/cmds/ddc.py
--- a/Engine/Extras/ushell/channels/unreal/core/cmds/ddc.py
+++ b/Engine/Extras/ushell/channels/unreal/core/cmds/ddc.py
@@ -36,8 +36,12 @@ class Auth(unrealcmd.Cmd):
         if sys.platform == "darwin": bin_type = "osx-x64"
         elif sys.platform == "linux": bin_type = "linux-x64"
 
+        bin_ext = ".exe"
+        if sys.platform == "darwin": bin_ext = ""
+        elif sys.platform == "linux": bin_ext = ""
+
         bin_path = engine.get_dir()
-        bin_path /= f"Binaries/DotNET/OidcToken/{bin_type}/OidcToken.exe"
+        bin_path /= f"Binaries/DotNET/OidcToken/{bin_type}/OidcToken{bin_ext}"
         if not bin_path.is_file():
             raise FileNotFoundError(f"Unable to find '{bin_path}'")
 
