Updated ushell [36979280]

#rnx

diff --git a/Engine/Extras/ushell/channels/flow/core/cmds/theme.py b/Engine/Extras/ushell/channels/flow/core/cmds/theme.py
--- a/Engine/Extras/ushell/channels/flow/core/cmds/theme.py
+++ b/Engine/Extras/ushell/channels/flow/core/cmds/theme.py
@@ -218,17 +218,19 @@ class Set(_ThemeCmd):
 #-------------------------------------------------------------------------------
 class Random(_ThemeCmd):
     """ Cycle randomly through all the themes """
-    mode = flow.cmd.Arg("", "Constrain randomness to dark or light themes")
+    mode = flow.cmd.Arg("", "Constrain randomness to dark/mid/light themes")
     one  = flow.cmd.Opt(False, "Pick one theme at random and exit")
 
-    complete_mode = ("dark", "light")
+    complete_mode = ("dark", "light", "mid")
 
     def main(self):
+        is_mid = lambda x: x.get_luminance() > 0.023 and x.get_luminance() < 0.3
         schemes = self.get_schemes()
         names = list(schemes.keys())
         what = self.args.mode.lower()
-        if what == "dark":  names = [x for x in names if not schemes[x].is_light()]
-        if what == "light": names = [x for x in names if     schemes[x].is_light()]
+        if   what == "dark":  names = [x for x in names if not schemes[x].is_light()]
+        elif what == "light": names = [x for x in names if     schemes[x].is_light()]
+        elif what == "mid":   names = [x for x in names if     is_mid(schemes[x].get_bg())]
 
         import random
         random.shuffle(names)
diff --git a/Engine/Extras/ushell/channels/flow/core/system/flow/_flick.py b/Engine/Extras/ushell/channels/flow/core/system/flow/_flick.py
--- a/Engine/Extras/ushell/channels/flow/core/system/flow/_flick.py
+++ b/Engine/Extras/ushell/channels/flow/core/system/flow/_flick.py
@@ -147,8 +147,8 @@ class Cmd(object):
 
         # Usage
         args = []
-        print("USAGE:", end="" if short else "\n")
-        print("", "[OPTIONS]", end="")
+        print("Usage:", end=" " if short else "\n  ")
+        print("[--<options>]", end="")
         for arg_name, arg in self.read_args():
             if arg.get_default() is None:
                 print("", arg_name, end="")
@@ -156,11 +156,13 @@ class Cmd(object):
                 print("", f"[{arg_name}]", end="")
             args.append((arg_name, arg.get_desc()))
 
-        if arg_many := self.get_many_arg():
-            print(f" [[--] {arg_many[0]}...]", end="")
-            args.append(("--", f"Delineates where '{arg_many[0]}' starts (optional)"))
-            args.append((arg_many[0], arg_many[1].get_desc()))
-        print()
+        suffix = ""
+        for arg_name, arg in self.read_many_args():
+            prefix = "[-- " if args else ""
+            print(f" {prefix}{arg_name}\u2026", end="")
+            suffix += "]" if args else ""
+            args.append((arg_name, arg.get_desc()))
+        print(suffix)
 
         if short:
             return
@@ -193,11 +195,11 @@ class Cmd(object):
                 leader = leader or (" " * col_0)
 
         if args:
-            print("\nARGS:")
+            print("\nArguments:")
             for name, desc in args:
                 print_param(name, desc)
 
-        print("\nOPTIONS:")
+        print("\nOptions:")
         for name, desc in opts:
             print_param(name, desc)
 
@@ -260,44 +262,70 @@ class Cmd(object):
                 positionals.append(arg)
 
         # Parse Args
-        index = -1
+        expect_dashdash = False
         cmd_arg = None
         cmd_args = self.read_args()
-        for index, (arg_name, cmd_arg) in enumerate(cmd_args):
-            if index >= len(positionals):
+        for arg_name, cmd_arg in cmd_args:
+            if not positionals:
                 return iter(((arg_name, cmd_arg), *cmd_args))
 
-            in_arg = positionals[index]
+            in_arg = positionals.pop(0)
             if in_arg == "--":
                 break
 
             result[arg_name] = (cmd_arg.cast_value(in_arg), False)
-        index += 1
+        else:
+            if cmd_arg:
+                if positionals:
+                    if positionals[0] != "--":
+                        raise ValueError(positionals)
+                    positionals.pop(0)
+                else:
+                    expect_dashdash = True
+
+        # Are there any many-arg sink for what remains?
+        cmd_many_arg = None
+        for arg_name, cmd_many_arg in self.read_many_args():
+            if not positionals:
+                if expect_dashdash:
+                    return iter((("--", None),))
+                return iter(((arg_name, cmd_many_arg),))
+            expect_dashdash = False
+
+            onwards = False
+            pieces = []
+            while positionals:
+                in_arg = positionals.pop(0)
+                if onwards := (in_arg == "--"):
+                    break
+                pieces.append(in_arg)
 
-        # Trim off a leading "--"
-        if len(positionals) > index and positionals[index] == "--":
-            index += 1
+            result[arg_name] = (tuple(pieces), False)
+
+            if not onwards and not positionals:
+                break
 
-        # Is there a many-arg sink for what remains?
-        many_cmd_arg = self.get_many_arg()
-        if many_cmd_arg:
-            arg_name, cmd_arg = many_cmd_arg
-            result[arg_name] = (tuple(positionals[index:]), len(positionals[index:]) == 0)
-            return iter(((arg_name, cmd_arg),))
+        if cmd_many_arg is not None:
+            if positionals:
+                result[arg_name] = (
+                    (*result[arg_name][0], "--", *positionals),
+                    False
+                )
+            return iter(((arg_name, cmd_many_arg),))
 
         # Check if there's anything left unparsed
-        if positionals[index:]:
-            raise ValueError(positionals[index:])
+        if positionals:
+            raise ValueError(positionals)
 
     def read_args(self):
         yield from self._read_argopts(_Arg)
 
+    def read_many_args(self):
+        yield from self._read_argopts(_ManyArg)
+
     def read_opts(self):
         yield from self._read_argopts(_Opt)
 
-    def get_many_arg(self):
-        return next(self._read_argopts(_ManyArg), None)
-
     @classmethod
     def get_name(self):
         return self.__name__
@@ -400,6 +428,10 @@ class Cmd(object):
 
         arg_name, argopt = next(arg_iter)
 
+        if arg_name == "--":
+            yield "--"
+            return
+
         completer = getattr(self, "complete_" + arg_name, None)
         if completer == None:
             raise LookupError
diff --git a/Engine/Extras/ushell/channels/unreal/core/cmds/stage.py b/Engine/Extras/ushell/channels/unreal/core/cmds/stage.py
--- a/Engine/Extras/ushell/channels/unreal/core/cmds/stage.py
+++ b/Engine/Extras/ushell/channels/unreal/core/cmds/stage.py
@@ -9,12 +9,12 @@ class _Impl(unrealcmd.Cmd):
     target   = unrealcmd.Arg(str, "Cooked target to stage")
     platform = unrealcmd.Arg("", "Platform whose cook should be staged")
     variant  = unrealcmd.Arg("development", "Build variant to stage binaries to copy")
-    style    = unrealcmd.Arg("auto", "Type of result to stage ('pak', 'nopak', 'zen', or 'auto' which is the default and makes the best choice for common use)")
+    style    = unrealcmd.Arg("auto", "Type of result to stage")
     uatargs  = unrealcmd.Arg([str], "Additional arguments to pass to UAT")
 
     complete_target = ("game", "client", "server")
 
-    complete_style = ("pak", "nopak", "zen","auto")
+    complete_style = ("pak", "nopak", "zen", "auto")
 
     def _build(self, build_editor, ubt_args=None):
         self.print_info("Building")
@@ -161,6 +161,13 @@ class Stage(_Impl):
     specifying the variant is only relevant if this stage is to be deployed at
     some point in the future (it is not the binary used when using '.run').
 
+    Available staging styles:
+
+      pak   - stage to pak/utoc/ucas files
+      nopak - stage for zen streaming if enabled, pak otherwise
+      zen   - stage for zen streaming, fail otherwise
+      auto  - [default] equivalent to nopak
+
     The --build/cook options accept an optional argument that are the arguments
     passed on to UnrealBuildTool and the cook commandlet respectively. For
     example; '--build="-GoSuperFast -DisableUnity"'. Double-up ("") or escape
diff --git a/Engine/Extras/ushell/channels/unreal/core/pylib/unreal/_context.py b/Engine/Extras/ushell/channels/unreal/core/pylib/unreal/_context.py
--- a/Engine/Extras/ushell/channels/unreal/core/pylib/unreal/_context.py
+++ b/Engine/Extras/ushell/channels/unreal/core/pylib/unreal/_context.py
@@ -2,6 +2,7 @@
 
 import os
 import re
+import sys
 from enum import Enum, Flag
 from typing import Iterator
 from pathlib import Path, WindowsPath, PosixPath
@@ -600,6 +601,17 @@ class Context(object):
     def __init__(self, primary_dir:str):
         primary_dir:Path = Path(primary_dir).resolve()
 
+        if sys.platform == "darwin":
+            # On Mac, we want to use the case-correct path, and the fcntl.F_GETPATH operation is the most direct way to get that
+            # Note that if we encounter an exception, the path will be left in its original form
+            try:
+                import fcntl
+                with os.fdopen(os.open(primary_dir, os.O_RDONLY)) as primary_dir_object:
+                    canonnical_path = fcntl.fcntl(primary_dir_object.fileno(), fcntl.F_GETPATH, bytes(os.fpathconf(primary_dir_object.fileno(), 'PC_PATH_MAX')))
+                    primary_dir:Path = Path(canonnical_path.decode('utf-8').rstrip('\x00')).resolve()
+            except:
+                pass
+
         # If a .uproject* was given instead the strip it off.
         if primary_dir.suffix.startswith(".uproject"):
             if primary_dir.is_file():
