Updated ushell [37033151]

#rnx

diff --git a/Engine/Extras/ushell/channels/flow/core/system/flow/_flick.py b/Engine/Extras/ushell/channels/flow/core/system/flow/_flick.py
--- a/Engine/Extras/ushell/channels/flow/core/system/flow/_flick.py
+++ b/Engine/Extras/ushell/channels/flow/core/system/flow/_flick.py
@@ -262,7 +262,6 @@ class Cmd(object):
                 positionals.append(arg)
 
         # Parse Args
-        expect_dashdash = False
         cmd_arg = None
         cmd_args = self.read_args()
         for arg_name, cmd_arg in cmd_args:
@@ -280,17 +279,12 @@ class Cmd(object):
                     if positionals[0] != "--":
                         raise ValueError(positionals)
                     positionals.pop(0)
-                else:
-                    expect_dashdash = True
 
         # Are there any many-arg sink for what remains?
         cmd_many_arg = None
         for arg_name, cmd_many_arg in self.read_many_args():
             if not positionals:
-                if expect_dashdash:
-                    return iter((("--", None),))
                 return iter(((arg_name, cmd_many_arg),))
-            expect_dashdash = False
 
             onwards = False
             pieces = []
@@ -428,10 +422,6 @@ class Cmd(object):
 
         arg_name, argopt = next(arg_iter)
 
-        if arg_name == "--":
-            yield "--"
-            return
-
         completer = getattr(self, "complete_" + arg_name, None)
         if completer == None:
             raise LookupError
