PR #12418: Make ushell look in /opt/homebrew for Python

#12418
#rnx

diff --git a/Engine/Extras/ushell/channels/flow/posix/provision.sh b/Engine/Extras/ushell/channels/flow/posix/provision.sh
--- a/Engine/Extras/ushell/channels/flow/posix/provision.sh
+++ b/Engine/Extras/ushell/channels/flow/posix/provision.sh
@@ -96,6 +96,7 @@ _bin_select \
     /usr/local/opt/python@${py_ver_maj}.*/bin/python$py_ver_maj \
     /usr/bin/python$py_ver_maj.*[0-9] \
     /usr/local/bin/python$py_ver_maj.*[0-9] \
+    /opt/homebrew/bin/python$py_ver_maj.*[0-9] \
     python
 
 if [ -z "$py_bin" ]; then
