Updated ushell [37225053]

#rnx

diff --git a/Engine/Extras/ushell/channels/flow/core/system/lib/bootstrap.py b/Engine/Extras/ushell/channels/flow/core/system/lib/bootstrap.py
--- a/Engine/Extras/ushell/channels/flow/core/system/lib/bootstrap.py
+++ b/Engine/Extras/ushell/channels/flow/core/system/lib/bootstrap.py
@@ -515,7 +515,10 @@ class _StateDir(object):
     _dir:Path
 
     def __init__(self, root:Path):
-        loc_hash = _hashen(str(Path(__file__).resolve()))
+        loc = str(Path(__file__).resolve())
+        if host := os.getenv("SHELL"):
+            loc += host
+        loc_hash = _hashen(loc)
         self._dir = root / ("flow_%08x" % loc_hash)
 
     def get_dir(self) -> Path:
diff --git a/Engine/Extras/ushell/channels/flow/nt/describe.flow.py b/Engine/Extras/ushell/channels/flow/nt/describe.flow.py
--- a/Engine/Extras/ushell/channels/flow/nt/describe.flow.py
+++ b/Engine/Extras/ushell/channels/flow/nt/describe.flow.py
@@ -4,9 +4,11 @@ import flow.describe
 
 #-------------------------------------------------------------------------------
 clink = flow.describe.Tool()
-clink.version("1.0.0a5")
-clink.payload("[REDACTED_PROFILE_URL]/clink/releases/download/$VERSION/clink-$VERSION.zip")
-clink.sha1("70289e92e3313a2b0e8dee801901eae61f8992a3")
+clink.version("1.0.0a6")
+#clink.payload("[REDACTED_PROFILE_URL]/clink/releases/download/$VERSION/clink-$VERSION.zip")
+#clink.sha1("70289e92e3313a2b0e8dee801901eae61f8992a3")
+clink.payload("https://www.dropbox.com/scl/fi/umxccnk4irsucu71z9keu/clink-1.0.0a6.7dc5c7.zip?rlkey=1l1f0u1gnt0mqertbpwc5zy2a&st=z6b9078t&dl=1")
+clink.sha1("6dd8010ed2118ce032e5ab4ca3697aecdfdc600a")
 clink.platform("win32")
 clink.bin("clink_x64.exe")
 
diff --git a/Engine/Extras/ushell/channels/flow/posix/shells.py b/Engine/Extras/ushell/channels/flow/posix/shells.py
--- a/Engine/Extras/ushell/channels/flow/posix/shells.py
+++ b/Engine/Extras/ushell/channels/flow/posix/shells.py
@@ -91,7 +91,7 @@ def _get_complete_bash():
 function _flow_prompt() {
     if [[ ! -z "$FLOW_PROMPT" ]]; then
         local prompt=${FLOW_PROMPT/:FLOW_CWD:/$(dirs +0)}
-        PS1=$(\$complete -- \$ "$prompt")
+        PS1=$(\$complete \$ "$prompt")
     fi
 }
 if [[ ! -z "$FLOW_PROMPT" ]]; then
@@ -99,7 +99,7 @@ if [[ ! -z "$FLOW_PROMPT" ]]; then
 fi
 
 function _flow_complete() {
-    local choices=$(\$complete -- ${COMP_WORDS[@]}...)
+    local choices=$(\$complete ${COMP_WORDS[@]}...)
     COMPREPLY=($(compgen -W "$choices" -- ${COMP_WORDS[$COMP_CWORD]}))
 }
 complete -F _flow_complete""" # purposely left no trailing CRLF
@@ -169,7 +169,7 @@ fi
 
 function _flow_complete() {
     # setopt local_options xtrace
-    compadd - $('$complete' -- ${^^words[@]:0:-1} $PREFIX...)
+    compadd - $('$complete' ${^^words[@]:0:-1} $PREFIX...)
 }
 compdef _flow_complete '$$$FLOW_CMDS$$$'
 """
diff --git a/Engine/Extras/ushell/channels/unreal/core/cmds/zen.py b/Engine/Extras/ushell/channels/unreal/core/cmds/zen.py
--- a/Engine/Extras/ushell/channels/unreal/core/cmds/zen.py
+++ b/Engine/Extras/ushell/channels/unreal/core/cmds/zen.py
@@ -25,6 +25,18 @@ def _get_ips_socket():
     ret = (_get_ip(socket.AF_INET6, "fe80::0"), _get_ip(socket.AF_INET, "172.31.255.255"))
     return [x for x in ret if x]
 
+def _get_ips_osx():
+    ret = _get_ips_socket()
+    try:
+        import subprocess
+        import re
+        segments = re.split("\\s+", subprocess.check_output('ifconfig | grep 169.254.', shell=True).decode('utf-8').strip())
+        print(segments)
+        ret.append(segments[1])
+    except:
+        pass
+    return ret
+
 def _get_ips_wmic():
     import subprocess as sp
     ret = []
@@ -45,7 +57,7 @@ def _get_ips_wmic():
     ret.sort(key=lambda x: "." in x) # sort ipv4 last
     return ret
 
-_get_ips = _get_ips_wmic if os.name == "nt" else _get_ips_socket
+_get_ips = _get_ips_wmic if os.name == "nt" else (_get_ips_osx if sys.platform == "darwin" else _get_ips_socket)
 
 
 
