Updated ushell [37832721]

#rnx

diff --git a/Engine/Extras/ushell/channels/flow/core/system/lib/bootstrap.py b/Engine/Extras/ushell/channels/flow/core/system/lib/bootstrap.py
--- a/Engine/Extras/ushell/channels/flow/core/system/lib/bootstrap.py
+++ b/Engine/Extras/ushell/channels/flow/core/system/lib/bootstrap.py
@@ -10,6 +10,8 @@ import fsutils
 import marshal
 import flow.describe
 import subprocess as sp
+import os
+import stat
 
 #-------------------------------------------------------------------------------
 class _Log(object):
@@ -97,6 +99,10 @@ def _http_get(url, dest_dir, progress_cb=None):
 #-------------------------------------------------------------------------------
 def _extract_tar(payload, dest_dir):
     cmd = ("tar", "-x", "-C", dest_dir, "-f", payload)
+    _, ext = os.path.splitext(payload)
+    if ext == '.zip' and os.name == 'posix':
+        # Tar on linux cannot open .zip files 
+        cmd = ("unzip", payload, '-d', dest_dir)
     ret = sp.run(cmd, stdout=sp.DEVNULL)
     assert ret.returncode == 0, "Extract failed; " + " ".join(cmd)
 
@@ -273,6 +279,12 @@ def _install_tool(name, tool):
         for bin_path in tool._bin_paths:
             if not os.path.exists(tool_dir + bin_path):
                 _log.write(f"Unable to find '{bin_path}' in '{tool_dir}'")
+            if os.name != "nt":
+                binary = os.path.join(tool_dir, bin_path)
+                # Confirm we have a binary that is executable.
+                if not (stat.S_IXUSR & os.stat(binary)[stat.ST_MODE]):
+                    st = os.stat(binary)
+                    os.chmod(binary, st.st_mode | stat.S_IEXEC)
     else:
         sha1 = 1
         if os.path.isfile(tool_dir + "/sha1.flow.txt"):
diff --git a/Engine/Extras/ushell/channels/unreal/core/pylib/fzf.py b/Engine/Extras/ushell/channels/unreal/core/pylib/fzf.py
--- a/Engine/Extras/ushell/channels/unreal/core/pylib/fzf.py
+++ b/Engine/Extras/ushell/channels/unreal/core/pylib/fzf.py
@@ -3,10 +3,11 @@
 import subprocess
 
 #-------------------------------------------------------------------------------
-def run(item_iter, *, height=10, prompt=">", multi=False, ansi=False, sort=True, ext_stdin=None):
+def run(item_iter, *, height=10, prompt=">", multi=False, ansi=False, sort=True, ext_stdin=None, exact=False):
     fzf_args = (
         "--prompt=" + prompt if prompt else None,
         None if sort else "--no-sort",
+        "--exact" if exact else None,
         "--multi" if multi == True else "--multi=" + str(multi) if multi > 0 else None,
         "--height=" + str(height + 2), # "+2" because fzf adds two lines of prompt
         "--color=16",
diff --git a/Engine/Extras/ushell/channels/unreal/core/pylib/unreal/_context.py b/Engine/Extras/ushell/channels/unreal/core/pylib/unreal/_context.py
--- a/Engine/Extras/ushell/channels/unreal/core/pylib/unreal/_context.py
+++ b/Engine/Extras/ushell/channels/unreal/core/pylib/unreal/_context.py
@@ -420,7 +420,7 @@ class _Project(_UfsNode):
     def get_engine(self) -> _Engine:
         return self._engine
 
-    def get_target_by_type(self, target_type:TargetType, *, create_temp:bool=True) -> _Target|None:
+    def get_target_by_type(self, target_type:TargetType, *, create_temp:bool=False) -> _Target|None:
         if target_type.value <= TargetType.PROGRAM.value:
             return
 
