Updated ushell [37881333]

#rnx

diff --git a/Engine/Extras/ushell/channels/unreal/core/pylib/unreal/platforms/win64/android.py b/Engine/Extras/ushell/channels/unreal/core/pylib/unreal/platforms/win64/android.py
--- a/Engine/Extras/ushell/channels/unreal/core/pylib/unreal/platforms/win64/android.py
+++ b/Engine/Extras/ushell/channels/unreal/core/pylib/unreal/platforms/win64/android.py
@@ -139,7 +139,7 @@ class Platform(unreal.Platform):
             os.system(" ".join(reverse_cmd))
 
         args = "'" + " ".join(x.replace("\"", "\\\"") for x in args) + "'"
-        cmd = (str(adb), "-d", "shell", "am", "start", "--es", "cmdline", args, intent)
+        cmd = (str(adb), "shell", "am", "start", "--es", "cmdline", args, intent)
         print(*cmd)
 
         # subprocess.run() messes with quotes on Windows so we'll use os.system
@@ -166,6 +166,6 @@ class Platform(unreal.Platform):
     def _kill(self, target=""):
         package_name = self._get_package_name()
         adb = self._get_adb()
-        cmd = (adb, "-d", "shell", "am", "force-stop", str(package_name))
+        cmd = (adb, "shell", "am", "force-stop", str(package_name))
         print(*cmd)
         sp.run(cmd)
