Updated ushell [38079492]

#rnx

diff --git a/Engine/Extras/ushell/channels/unreal/core/pylib/unreal/_context.py b/Engine/Extras/ushell/channels/unreal/core/pylib/unreal/_context.py
--- a/Engine/Extras/ushell/channels/unreal/core/pylib/unreal/_context.py
+++ b/Engine/Extras/ushell/channels/unreal/core/pylib/unreal/_context.py
@@ -420,10 +420,13 @@ class _Project(_UfsNode):
     def get_engine(self) -> _Engine:
         return self._engine
 
-    def get_target_by_type(self, target_type:TargetType, *, create_temp:bool=False) -> _Target|None:
+    def get_target_by_type(self, target_type:TargetType) -> _Target|None:
         if target_type.value <= TargetType.PROGRAM.value:
             return
 
+        if not (self.get_dir() / "Source").is_dir():
+            return self.get_parent_node().get_target_by_type(target_type)
+
         for target in self._read_targets("Source", ancestors=False, shallow=True):
             if target.get_type() == target_type:
                 return target
@@ -439,39 +442,6 @@ class _Project(_UfsNode):
 
         return target
 
-    def _create_temp_target(self, target:_Target) -> None:
-        int_dir = self.get_dir() / "Source"
-
-        # Only do creation if we are the first to arrive
-        int_tag = int_dir / ".ushell_created"
-        if not int_dir.is_dir():
-            int_dir.mkdir(parents=True, exist_ok=True)
-            int_tag.open("wb").close()
-        else:
-            if not int_tag.is_file():
-                return
-
-        target_name = target.get_name()
-        dot_cs_path = int_dir / (target_name + ".Target.cs")
-        if dot_cs_path.is_file():
-            return
-
-        game_target = "UnrealGame"
-        if self._engine.get_version_major() <= 4:
-            game_target = "UE4Game"
-
-        type_name = target.get_type().name.title()
-        with dot_cs_path.open("wt") as dot_cs:
-            dot_cs.write(
-                "// Created by ushell to make blueprint projects work end to end.\n"
-                "// Placed in Source/ because UAT cleans Intermediate/Source/\n"
-                "using UnrealBuildTool;\n"
-                "public class " + target_name + "Target : TargetRules {\n"
-                "public " + target_name + "Target(TargetInfo Target) : base(Target) {\n"
-                "Type = TargetType." + type_name + ";\n"
-                "ExtraModuleNames.AddRange(new string[] {\"" + game_target + "\"} ); }}\n"
-            )
-
 
 
 # {{{1 branch ------------------------------------------------------------------
diff --git a/Engine/Extras/ushell/channels/unreal/core/pylib/unreal/platforms/win64/win64.py b/Engine/Extras/ushell/channels/unreal/core/pylib/unreal/platforms/win64/win64.py
--- a/Engine/Extras/ushell/channels/unreal/core/pylib/unreal/platforms/win64/win64.py
+++ b/Engine/Extras/ushell/channels/unreal/core/pylib/unreal/platforms/win64/win64.py
@@ -2,6 +2,7 @@
 
 import os
 import unreal
+import unreal.cmdline
 import subprocess as sp
 
 #-------------------------------------------------------------------------------
@@ -119,13 +120,16 @@ class Platform(_PlatformBase):
 
     def _launch(self, exec_context, stage_dir, binary_path, args):
         if stage_dir:
-            midfix = "Engine";
+            midfixes = ("Engine",);
             if project := self.get_unreal_context().get_project():
-                midfix = project.get_name()
-            base_dir = stage_dir + midfix + "/Binaries/Win64"
-            if not os.path.isdir(base_dir):
-                raise EnvironmentError(f"Failed to find base directory '{base_dir}'")
-            args = (*args, "-basedir=" + base_dir)
+                midfixes = (project.get_name(), *midfixes)
+            for midfix in midfixes:
+                base_dir = stage_dir + midfix + "/Binaries/Win64"
+                if os.path.isdir(base_dir):
+                    break
+            else:
+                raise EnvironmentError(f"Failed to find binaries directory from '{stage_dir}'")
+            args = (*args, *unreal.cmdline.read_ueified("-basedir=" + base_dir))
 
         cmd = exec_context.create_runnable(binary_path, *args)
         cmd.launch()
