Updated ushell [38377678]

#rnx

diff --git a/Engine/Extras/ushell/channels/flow/core/system/flow/describe.py b/Engine/Extras/ushell/channels/flow/core/system/flow/describe.py
--- a/Engine/Extras/ushell/channels/flow/core/system/flow/describe.py
+++ b/Engine/Extras/ushell/channels/flow/core/system/flow/describe.py
@@ -5,26 +5,31 @@ from . import _backtrace
 #-------------------------------------------------------------------------------
 class Tool(object):
     class _Bundle(object):
-        def __init__(self):                    self._payloads = []
-        def payload(self, url, filetype=None): self._payloads.append((url, filetype))
+        def __init__(self):                       self._payloads = []; self._bin_paths = []
+        def platform(self, platform):             self._platform = platform
+        def bin(self, bin_path):                  self._bin_paths.append(bin_path)
+        def sha1(self, sha1):                     self._sha1 = sha1
+        def payload(self, url, *, filetype=None): self._payloads.append((url, filetype))
 
     def __init__(self):
-        self._bin_paths = []
         self._bundles = {}
         self._root_dir = None
-        self._sha1 = None
         self._version = None
-        self._platform = None
         self._source = (None, None)
 
-    def sha1(self, sha1):                  self._sha1 = sha1
-    def version(self, version):            self._version = version
-    def platform(self, platform):          self._platform = platform
-    def source(self, url, ver_regex):      self._source = (url, ver_regex)
-    def bin(self, bin_path):               self._bin_paths.append(bin_path)
-    def root_dir(self, root_dir):          self._root_dir = root_dir
-    def bundle(self, name="$"):            return self._bundles.setdefault(name, Tool._Bundle())
-    def payload(self, url, filetype=None): self.bundle().payload(url, filetype)
+    def bundle(self, name="$", *, platform=None):
+        if platform: name += f"-{platform}"
+        bundle = self._bundles.setdefault(name, Tool._Bundle())
+        if platform: bundle.platform(platform)
+        return bundle
+
+    def version(self, version):               self._version = version
+    def source(self, url, ver_regex):         self._source = (url, ver_regex)
+    def root_dir(self, root_dir):             self._root_dir = root_dir
+    def payload(self, url, *, filetype=None): self.bundle().payload(url, filetype=filetype)
+    def platform(self, platform):             self.bundle().platform(platform)
+    def bin(self, bin_path):                  self.bundle().bin(bin_path)
+    def sha1(self, sha1):                     self.bundle().sha1(sha1)
 
     # Implemented elsewhere.
     #def extract(self, src_path, dest_dir):
diff --git a/Engine/Extras/ushell/channels/flow/core/system/flow/system.py b/Engine/Extras/ushell/channels/flow/core/system/flow/system.py
--- a/Engine/Extras/ushell/channels/flow/core/system/flow/system.py
+++ b/Engine/Extras/ushell/channels/flow/core/system/flow/system.py
@@ -139,13 +139,13 @@ class _Tool(_FlowObject):
         return self._name
 
     def get_dir(self):
-        return self.get_system().get_tools_dir() + f"{self._name}-{self._data[0]}/"
+        return self.get_system().get_tools_dir() + self._data["double"]
 
     def get_version(self):
-        return self._data[0]
+        return self._data["version"]
 
     def get_bin_paths(self):
-        return self._data[1]
+        return self._data["bin_paths"]
 
 
 
diff --git a/Engine/Extras/ushell/channels/flow/core/system/lib/bootstrap.py b/Engine/Extras/ushell/channels/flow/core/system/lib/bootstrap.py
--- a/Engine/Extras/ushell/channels/flow/core/system/lib/bootstrap.py
+++ b/Engine/Extras/ushell/channels/flow/core/system/lib/bootstrap.py
@@ -141,24 +141,27 @@ def _validate_tool(name, tool):
     assert tool._bundles, f"No payload set on tool '{name}'"
 
 #-------------------------------------------------------------------------------
-def _acquire_tool(name, tool, target_dir, progress_cb):
+def _acquire_tool(name, tool, manifest, target_dir, progress_cb):
     import hashlib
 
+    bundles = {k:v for k,v in manifest["bundles"].items() if v["enabled"]}
+    if not bundles:
+        return
+
     # Some boiler plate to provide a progress update
     payload_i = 0
-    payload_n = sum(len(x._payloads) for x in tool._bundles.values())
+    payload_n = sum(len(x["payloads"]) for x in bundles.values())
     def inner_progress_cb(file_name, value):
         value = int(((float(payload_i) + value) * 100) / payload_n)
         progress_cb(f"{value:3}%")
 
     # Acquire the tool's bundles/payloads
     temp_dir = fsutils.WorkPath(target_dir[:-1] + ".work/")
-    for bundle_name, bundle in tool._bundles.items():
+    for bundle_name, bundle in bundles.items():
         _log.write("Bundle", bundle_name)
         fetch_dir = temp_dir + bundle_name + "/"
         os.mkdir(fetch_dir)
-        for payload, ftype in bundle._payloads:
-            payload = payload.replace("$VERSION", tool._version)
+        for payload, ftype in bundle["payloads"]:
             _log.write("Payload", payload, ftype)
             payload_dest = _http_get(payload, fetch_dir, inner_progress_cb)
             payload_i += 1
@@ -167,8 +170,8 @@ def _acquire_tool(name, tool, target_dir, progress_cb):
 
     # Hash the downloaded content to confirm it is what is expected
     progress_cb("sha1")
-    sha1 = hashlib.sha1()
-    for bundle_name, bundle in tool._bundles.items():
+    for bundle_name, bundle in bundles.items():
+        sha1 = hashlib.sha1()
         for item in fsutils.read_files(temp_dir + bundle_name):
             with open(item.path, "rb") as in_file:
                 while True:
@@ -176,19 +179,21 @@ def _acquire_tool(name, tool, target_dir, progress_cb):
                     if len(data) <= 0:
                         break
                     sha1.update(data)
-    tool_sha1 = sha1.hexdigest()
-    _log.write("SHA1", name, tool_sha1)
+        sha1 = sha1.hexdigest()
+        _log.write("SHA1", bundle_name, sha1)
+
+        if expected_sha1 := bundle["sha1"]:
+            expected_sha1 = expected_sha1.lower()
+            if sha1 != expected_sha1:
+                assert false, f"Unexpected content bundle data for tool '{name}' [sha1:{sha1}]"
 
-    if tool._sha1:
-        expected_sha1 = tool._sha1.lower()
-        if tool_sha1 != expected_sha1:
-            assert false, f"Unexpected content bundle data for tool '{name}' [sha1:{tool_sha1}]"
+        bundle["sha1"] = sha1
 
     # Extract the payloads.
     progress_cb("uzip")
     extract_dir = temp_dir + "$extract/"
     os.mkdir(extract_dir)
-    for bundle_name, bundle in tool._bundles.items():
+    for bundle_name, bundle in bundles.items():
         _log.write("Extract", bundle_name, "from", fetch_dir, "to", extract_dir, "with", tool.extract)
         fetch_dir = temp_dir + bundle_name + "/"
         tool.extract(fetch_dir, extract_dir)
@@ -207,92 +212,127 @@ def _acquire_tool(name, tool, target_dir, progress_cb):
         source_dir += "/" + tool._root_dir
         assert os.path.isdir(source_dir), "Given root f'{tool._root_dir}' for tool '{name}' is not valid"
 
-    # Store sha1
-    with open(source_dir + "/sha1.flow.txt", "w") as out:
-        out.write(tool_sha1)
-
-    os.rename(source_dir, target_dir)
+    # Windows can take its time moving/deleting a directory causing the rename
+    # to spuriously fail. Or it AV interfering with source. So retry a few times.
+    for i in range(6):
+        try:
+            os.rename(source_dir, target_dir)
+            break
+        except OSError:
+            if i == 5:
+                raise
+            import time
+            time.sleep(0.1)
+            continue
 
     progress_cb("done")
-    return tool_sha1
 
 #-------------------------------------------------------------------------------
 def _manifest_tool(name, tool):
+    def is_enabled(bundle):
+        bundle_platform = getattr(bundle, "_platform", None)
+        if not bundle_platform:
+            return True
+
+        host_platform = sys.platform
+        if "-" in bundle_platform:
+            import platform
+            arch = platform.machine()
+            arch = "amd64" if arch in ("x86_64", "AMD64") else "arm64"
+            host_platform += "-" + arch
+
+        return bundle_platform == host_platform
+
     bundles = {}
     for bundle_name, bundle in tool._bundles.items():
-        urls = (x[0] for x in bundle._payloads)
-        bundles[bundle_name] = tuple(x.replace("$VERSION", tool._version) for x in urls)
+        payloads = ((u.replace("$VERSION", tool._version),ft) for u,ft in bundle._payloads)
+        bundles[bundle_name] = {
+            "enabled"   : is_enabled(bundle),
+            "payloads"  : tuple(payloads),
+            "bin_paths" : tuple(bundle._bin_paths),
+            "sha1"      : getattr(bundle, "_sha1", 0),
+        }
 
     return {
-        "sha1"      : 0,
         "version"   : tool._version,
         "double"    : f"{name}-{tool._version}",
-        "bin_paths" : tuple(tool._bin_paths),
-        "bundles"   : bundles,
         "source"    : tool._source,
+        "bundles"   : bundles,
+        "bin_paths" : tuple(),
     }
 
 #-------------------------------------------------------------------------------
-def _install_tool(name, tool):
+def _install_tool(name, tool, cleaner):
     manifest = _manifest_tool(name, tool)
 
-    if req_platform := tool._platform:
-        this_platform = sys.platform
-        if "-" in req_platform:
-            import platform
-            arch = platform.machine()
-            arch = "amd64" if arch in ("x86_64", "AMD64") else "arm64"
-            this_platform += "-" + arch
-        if req_platform != this_platform:
-            return manifest
-
     tool_double = f"{name}-{tool._version}/"
+    tool_dir = "../tools/" + tool_double
+    manifest_path = tool_dir + "manifest.1.flow"
     _log.write("Installing tool", tool_double)
 
+    # If the tool exists then reuse the cached manifest
+    if os.path.isfile(manifest_path):
+        _log.write("Tool already acquired", tool_dir)
+        with open(manifest_path, "rb") as inp:
+            return marshal.load(inp)
+
+    # Clean up any legacy deployments
+    if os.path.isdir(tool_dir):
+        cleaner.delete(Path(tool_dir))
+
+    # Acquire the tool
     log_line = (13 * " ") + tool_double[:-1]
     log_line += "\b" * (len(log_line) - 1)
     _log.print(log_line, end="\r      [    ]\b")
 
-    # If the tool doesn't exist then it needs acquiring
-    tool_dir = "../tools/" + tool_double
-    if not os.path.isdir(tool_dir):
-        def progress_cb(info):
-            _log.print("\b" * 4, info[:4], sep="", end="")
+    def progress_cb(info):
+        _log.print("\b" * 4, info[:4], sep="", end="")
 
-        try:
-            _log.write("Acquiring to", tool_dir)
-            sha1 = _acquire_tool(name, tool, tool_dir, progress_cb)
-        except KeyboardInterrupt:
-            raise
-        except Exception as e:
-            _log.write("ERROR", str(e))
-            _log.print("\b" * 4, "fail", sep="", end="\n")
-            return manifest
-
-        # Post install
-        if hasattr(tool, "post_install"):
-            progress_cb("post")
-            _log.write("Post-install")
-            tool.post_install(tool_dir)
-
-        # Validate that the required binaries exist.
-        for bin_path in tool._bin_paths:
+    try:
+        _log.write("Acquiring to", tool_dir)
+        _acquire_tool(name, tool, manifest, tool_dir, progress_cb)
+    except KeyboardInterrupt:
+        raise
+    except Exception as e:
+        _log.write("ERROR", str(e))
+        _log.print("\b" * 4, "fail", sep="", end="\n")
+        return manifest
+
+    # Maybe this tool doen't deploy for this platform?
+    if not os.path.isdir(tool_dir):
+        _log.print("\b" * 4, " na ", sep="", end="\n")
+        return manifest
+
+    # Post install
+    if hasattr(tool, "post_install"):
+        progress_cb("post")
+        _log.write("Post-install")
+        tool.post_install(tool_dir)
+
+    # Validate that the required binaries exist.
+    bin_paths = set()
+    for bundle in manifest["bundles"].values():
+        if not bundle["enabled"]:
+            continue
+        for bin_path in bundle["bin_paths"]:
             if not os.path.exists(tool_dir + bin_path):
                 _log.write(f"Unable to find '{bin_path}' in '{tool_dir}'")
-            if os.name != "nt":
-                binary = os.path.join(tool_dir, bin_path)
-                # Confirm we have a binary that is executable.
-                if not (stat.S_IXUSR & os.stat(binary)[stat.ST_MODE]):
-                    st = os.stat(binary)
-                    os.chmod(binary, st.st_mode | stat.S_IEXEC)
-    else:
-        sha1 = 1
-        if os.path.isfile(tool_dir + "/sha1.flow.txt"):
-            with open(tool_dir + "/sha1.flow.txt", "rt") as sha1_file:
-                sha1 = sha1_file.read()
-        _log.write("Tool already acquired", tool_dir)
+                continue
+
+            bin_paths.add(bin_path)
+
+            # Confirm we have a binary that is executable.
+            if os.name == "nt":
+                continue
+            binary = os.path.join(tool_dir, bin_path)
+            if not (stat.S_IXUSR & os.stat(binary)[stat.ST_MODE]):
+                st = os.stat(binary)
+                os.chmod(binary, st.st_mode | stat.S_IEXEC)
+    manifest["bin_paths"] = tuple(bin_paths)
+
+    with open(manifest_path, "wb") as out:
+        marshal.dump(manifest, out)
 
-    manifest["sha1"] = sha1
     _log.print("\b" * 4, " ok ", sep="", end="\n")
     return manifest
 
@@ -306,7 +346,7 @@ def _validate_command(channel_dir, name, command):
         assert False, f"Command '{name}' missing source '{command._py_path}'"
 
 #-------------------------------------------------------------------------------
-def _build_channel(channel_name:str, channel_dir:Path) -> dict[str]:
+def _build_channel(channel_name:str, channel_dir:Path, cleaner:"_Cleaner") -> dict[str]:
     _log.print(f"Updating channel '{channel_name}'")
     _log.indent("Channel", channel_name)
 
@@ -346,7 +386,7 @@ def _build_channel(channel_name:str, channel_dir:Path) -> dict[str]:
             _validate_tool(name, tool)
 
         for name, tool in tools:
-            tools_manifest[name] = _install_tool(name, tool)
+            tools_manifest[name] = _install_tool(name, tool, cleaner)
         _log.unindent()
 
     # Commands
@@ -398,12 +438,11 @@ def _carve_channels(channels):
     for channel in channels:
         index = channel["index"]
         parent_index = channel["parent"]["index"] if channel["parent"] else -1
-        tools = {k:(v["version"], v["bin_paths"]) for k,v in channel["tools"].items() if v["sha1"]}
-        out[channel["index"]] = {
+        out[index] = {
             "name" : channel["name"],
             "path" : channel["path"],
             "parent" : parent_index,
-            "tools" : tools,
+            "tools" : channel["tools"],
         }
 
     return out
@@ -431,8 +470,10 @@ def _create_shims(manifest):
     # Create tool shims
     tool_root = os.path.abspath("../tools") + "/"
     for channel in reversed(manifest["channels"]):
-        for name, (version, bin_paths) in channel["tools"].items():
-            tool_dir = tool_root + f"{name}-{version}/"
+        for name, tool in channel["tools"].items():
+            version = tool["version"]
+            bin_paths = tool["bin_paths"]
+            tool_dir = tool_root + tool["double"] + "/"
             for bin_path in bin_paths:
                 name = os.path.basename(bin_path)
                 tool_path = tool_dir + bin_path
@@ -609,7 +650,7 @@ def _get_sources() -> _Sources:
     return sources
 
 #-------------------------------------------------------------------------------
-def _update(working_dir:Path, sources:_Sources) -> None:
+def _update(working_dir:Path, sources:_Sources, cleaner:"_Cleaner") -> None:
     class _CwdScope(object):
         def __init__(self, dir:Path): self._prev_dir = os.getcwd(); os.chdir(dir)
         def __del__(self) -> None:    os.chdir(self._prev_dir)
@@ -627,7 +668,7 @@ def _update(working_dir:Path, sources:_Sources) -> None:
         # Build all the channels
         manifests = []
         for name, flow_py_path in sources:
-            manifest = _build_channel(name, flow_py_path.parent)
+            manifest = _build_channel(name, flow_py_path.parent, cleaner)
             manifests.append(manifest)
 
         # Aggregate all channels into a manifest and create command shims
@@ -671,7 +712,7 @@ def impl() -> Path:
     temp_dir = dest_dir.parent / (dest_dir.name + f"${os.getpid()}")
     try:
         temp_dir.mkdir(parents=True)
-        _update(temp_dir, sources)
+        _update(temp_dir, sources, cleaner)
         dest_dir = temp_dir.rename(dest_dir)
     except:
         cleaner.delete(temp_dir)
diff --git a/Engine/Extras/ushell/channels/flow/core/system/version b/Engine/Extras/ushell/channels/flow/core/system/version
--- a/Engine/Extras/ushell/channels/flow/core/system/version
+++ b/Engine/Extras/ushell/channels/flow/core/system/version
@@ -1 +1 @@
-21
+22
diff --git a/Engine/Extras/ushell/channels/flow/nt/describe.flow.py b/Engine/Extras/ushell/channels/flow/nt/describe.flow.py
--- a/Engine/Extras/ushell/channels/flow/nt/describe.flow.py
+++ b/Engine/Extras/ushell/channels/flow/nt/describe.flow.py
@@ -22,13 +22,13 @@ fd.bin("fd.exe")
 fd.source("[REDACTED_PROFILE_URL]/fd/releases/latest", "fd-v([0-9.]+)-x86_64")
 
 #-------------------------------------------------------------------------------
-fzf_win32 = flow.describe.Tool()
-fzf_win32.version("0.54.3")
-fzf_win32.payload("[REDACTED_PROFILE_URL]/fzf/releases/download/v$VERSION/fzf-$VERSION-windows_amd64.zip")
-fzf_win32.sha1("2c3ceab1a099ad3ab30eaffb6edaaa8f1a8ce2b4")
-fzf_win32.platform("win32")
-fzf_win32.bin("fzf.exe")
-fzf_win32.source("[REDACTED_PROFILE_URL]/fzf/releases/latest", r"fzf-(\d+\.\d+\.\d+)-windows")
+fzf = flow.describe.Tool()
+fzf.version("0.54.3")
+fzf.source("[REDACTED_PROFILE_URL]/fzf/releases/latest", r"fzf-(\d+\.\d+\.\d+)-windows")
+if bundle := fzf.bundle(platform="win32"):
+    bundle.payload("[REDACTED_PROFILE_URL]/fzf/releases/download/v$VERSION/fzf-$VERSION-windows_amd64.zip")
+    bundle.bin("fzf.exe")
+    bundle.sha1("2c3ceab1a099ad3ab30eaffb6edaaa8f1a8ce2b4")
 
 #-------------------------------------------------------------------------------
 shell_cmd = flow.describe.Command()
diff --git a/Engine/Extras/ushell/channels/unreal/core/describe.flow.py b/Engine/Extras/ushell/channels/unreal/core/describe.flow.py
--- a/Engine/Extras/ushell/channels/unreal/core/describe.flow.py
+++ b/Engine/Extras/ushell/channels/unreal/core/describe.flow.py
@@ -1,61 +1,45 @@
 # Copyright Epic Games, Inc. All Rights Reserved.
 
-import sys
 import flow.describe
 
 #-------------------------------------------------------------------------------
-fzf_linux = flow.describe.Tool()
-fzf_linux.version("0.54.3")
-fzf_linux.payload("[REDACTED_PROFILE_URL]/fzf/releases/download/v$VERSION/fzf-$VERSION-linux_amd64.tar.gz")
-fzf_linux.sha1("e46fcd29c453ee61be0d41ca22b7eefd944f91f6")
-fzf_linux.platform("linux")
-fzf_linux.bin("fzf")
-
-fzf_darwin_x64 = flow.describe.Tool()
-fzf_darwin_x64.version("0.54.3")
-fzf_darwin_x64.payload("[REDACTED_PROFILE_URL]/fzf/releases/download/v$VERSION/fzf-$VERSION-darwin_amd64.tar.gz")
-fzf_darwin_x64.sha1("3723b3445037bdaa302d1d365d4940beb4c283c8")
-fzf_darwin_x64.platform("darwin-amd64")
-fzf_darwin_x64.bin("fzf")
-
-fzf_darwin_arm = flow.describe.Tool()
-fzf_darwin_arm.version("0.54.3")
-fzf_darwin_arm.payload("[REDACTED_PROFILE_URL]/fzf/releases/download/v$VERSION/fzf-$VERSION-darwin_arm64.tar.gz")
-fzf_darwin_arm.sha1("c52657e62c095909fb1ac69e6d714eac3a10dc4c")
-fzf_darwin_arm.platform("darwin-arm64")
-fzf_darwin_arm.bin("fzf")
+fzf = flow.describe.Tool()
+fzf.version("0.54.3")
+if bundle := fzf.bundle(platform="linux"):
+    bundle.payload("[REDACTED_PROFILE_URL]/fzf/releases/download/v$VERSION/fzf-$VERSION-linux_amd64.tar.gz")
+    bundle.bin("fzf")
+    bundle.sha1("e46fcd29c453ee61be0d41ca22b7eefd944f91f6")
+if bundle := fzf.bundle(platform="darwin-amd64"):
+    bundle.payload("[REDACTED_PROFILE_URL]/fzf/releases/download/v$VERSION/fzf-$VERSION-darwin_amd64.tar.gz")
+    bundle.bin("fzf")
+    bundle.sha1("3723b3445037bdaa302d1d365d4940beb4c283c8")
+if bundle := fzf.bundle(platform="darwin-arm64"):
+    bundle.payload("[REDACTED_PROFILE_URL]/fzf/releases/download/v$VERSION/fzf-$VERSION-darwin_arm64.tar.gz")
+    bundle.bin("fzf")
+    bundle.sha1("c52657e62c095909fb1ac69e6d714eac3a10dc4c")
 
 
 
 #-------------------------------------------------------------------------------
-ripgrep_win32 = flow.describe.Tool()
-ripgrep_win32.version("14.1.0")
-ripgrep_win32.payload("[REDACTED_PROFILE_URL]/ripgrep/releases/download/$VERSION/ripgrep-$VERSION-x86_64-pc-windows-msvc.zip")
-ripgrep_win32.sha1("ba7820c05aeb49050d9242a109e422585c3fc767")
-ripgrep_win32.platform("win32")
-ripgrep_win32.bin("rg.exe")
-ripgrep_win32.source("[REDACTED_PROFILE_URL]/ripgrep/releases/latest", r"ripgrep-(\d+\.\d+\.\d+)-x86_64")
-
-ripgrep_linux = flow.describe.Tool()
-ripgrep_linux.version("14.1.0")
-ripgrep_linux.payload("[REDACTED_PROFILE_URL]/ripgrep/releases/download/$VERSION/ripgrep-$VERSION-x86_64-unknown-linux-musl.tar.gz")
-ripgrep_linux.sha1("8267841999fa96f5a5a375ed518a24bebef6ca4c")
-ripgrep_linux.platform("linux")
-ripgrep_linux.bin("rg")
-
-ripgrep_darwin_x64 = flow.describe.Tool()
-ripgrep_darwin_x64.version("14.1.0")
-ripgrep_darwin_x64.payload("[REDACTED_PROFILE_URL]/ripgrep/releases/download/$VERSION/ripgrep-$VERSION-x86_64-apple-darwin.tar.gz")
-ripgrep_darwin_x64.sha1("0d98d699f08c4eed18f4379734d4e1d973f59cde")
-ripgrep_darwin_x64.platform("darwin-amd64")
-ripgrep_darwin_x64.bin("rg")
-
-ripgrep_darwin_arm = flow.describe.Tool()
-ripgrep_darwin_arm.version("14.1.0")
-ripgrep_darwin_arm.payload("[REDACTED_PROFILE_URL]/ripgrep/releases/download/$VERSION/ripgrep-$VERSION-aarch64-apple-darwin.tar.gz")
-ripgrep_darwin_arm.sha1("59d0d543a7e25831b29ebe2943da366d4c005d26")
-ripgrep_darwin_arm.platform("darwin-arm64")
-ripgrep_darwin_arm.bin("rg")
+ripgrep = flow.describe.Tool()
+ripgrep.version("14.1.0")
+ripgrep.source("[REDACTED_PROFILE_URL]/ripgrep/releases/latest", r"ripgrep-(\d+\.\d+\.\d+)-x86_64")
+if bundle := ripgrep.bundle(platform="win32"):
+    bundle.payload("[REDACTED_PROFILE_URL]/ripgrep/releases/download/$VERSION/ripgrep-$VERSION-x86_64-pc-windows-msvc.zip")
+    bundle.bin("rg.exe")
+    bundle.sha1("ba7820c05aeb49050d9242a109e422585c3fc767")
+if bundle := ripgrep.bundle(platform="linux"):
+    bundle.payload("[REDACTED_PROFILE_URL]/ripgrep/releases/download/$VERSION/ripgrep-$VERSION-x86_64-unknown-linux-musl.tar.gz")
+    bundle.bin("rg")
+    bundle.sha1("8267841999fa96f5a5a375ed518a24bebef6ca4c")
+if bundle := ripgrep.bundle(platform="darwin-amd64"):
+    bundle.payload("[REDACTED_PROFILE_URL]/ripgrep/releases/download/$VERSION/ripgrep-$VERSION-x86_64-apple-darwin.tar.gz")
+    bundle.bin("rg")
+    bundle.sha1("0d98d699f08c4eed18f4379734d4e1d973f59cde")
+if bundle := ripgrep.bundle(platform="darwin-arm64"):
+    bundle.payload("[REDACTED_PROFILE_URL]/ripgrep/releases/download/$VERSION/ripgrep-$VERSION-aarch64-apple-darwin.tar.gz")
+    bundle.bin("rg")
+    bundle.sha1("59d0d543a7e25831b29ebe2943da366d4c005d26")
 
 
 
diff --git a/Engine/Extras/ushell/channels/unreal/perforce/cmds/bisect.py b/Engine/Extras/ushell/channels/unreal/perforce/cmds/bisect.py
--- a/Engine/Extras/ushell/channels/unreal/perforce/cmds/bisect.py
+++ b/Engine/Extras/ushell/channels/unreal/perforce/cmds/bisect.py
@@ -55,6 +55,9 @@ class _Bisectomatron(object):
 
 #-------------------------------------------------------------------------------
 def _resolve_script_path(script_path):
+    if os.name != "nt" and script_path.endswith(".sh"):
+        raise ValueError(f"A suitable script interpreter should also specified; 'bash {script_path}'");
+
     import shutil
     if whiched := shutil.which(script_path):
         return Path(whiched)
@@ -77,7 +80,7 @@ class Bisect(unrealcmd.Cmd):
     is considered next is controlled by the scripts exit code; 0=good, 80=bad, and
     90=failed_build. The exit codes are chosen to avoid interference from general
     scripting errors. Here is an example Batch script that could be used with the
-    command line '.p4 bisect 123 456 path/to/script.bat test_case debug';
+    command line '.p4 bisect 123 456 -- path/to/script.bat test_case debug';
 
       @echo off
       .build editor %2
