Updated ushell [38752955]

#rnx

diff --git a/Engine/Extras/ushell/channels/flow/core/boot.py b/Engine/Extras/ushell/channels/flow/core/boot.py
--- a/Engine/Extras/ushell/channels/flow/core/boot.py
+++ b/Engine/Extras/ushell/channels/flow/core/boot.py
@@ -12,9 +12,22 @@ class Boot(flow.cmd.Cmd):
     cleanprompt = flow.cmd.Opt(False, "Do not modify the host shell's prompt")
     theme       = flow.cmd.Opt("", "Adds a spot of colour with '.ushell theme'")
 
+    def print_title(self):
+        lines = (
+            "         _       _ _ ",
+            " _ _ ___| |_ ___| | |",
+            "| | |_ -|   | -_| | |",
+            "|___|___|_|_|___|_|_|",
+        )
+        print("\x1b[93m", end="")
+        for i, line in enumerate(lines):
+            print("  ", line, sep="")
+        print("\n\x1b[0m", end="")
+
     def main(self):
         # Clear the screen
         sys.stdout.write("\x1b[2J\x1b[0;0H")
+        self.print_title()
 
         # Change to the start directory the user requested. Nice and early.
         if self.args.startdir:
diff --git a/Engine/Extras/ushell/channels/flow/core/cmds/debug.py b/Engine/Extras/ushell/channels/flow/core/cmds/debug.py
--- a/Engine/Extras/ushell/channels/flow/core/cmds/debug.py
+++ b/Engine/Extras/ushell/channels/flow/core/cmds/debug.py
@@ -176,15 +176,15 @@ class Sha1s(_Channels):
                     print("%.02f%%" % ((recv_size * 100) / total_size), end="\r")
                 digest.update(data)
 
-            digest = hashlib.sha1()
             for bundle_name, bundle in tool._bundles.items():
+                digest = hashlib.sha1()
                 for url, file_type in bundle._payloads:
                     url = url.replace("$VERSION", tool._version)
                     try: _http_get(url, on_data)
                     except URLError: pass
-            sha1 = digest.hexdigest()
+                sha1 = digest.hexdigest()
 
-            print(rf'{tool_name}.sha1("{sha1}")')
+                print(rf'{tool_name}{bundle_name}.sha1("{sha1}")')
 
     def main(self):
         super().main()
diff --git a/Engine/Extras/ushell/channels/flow/core/cmds/tips.py b/Engine/Extras/ushell/channels/flow/core/cmds/tips.py
--- a/Engine/Extras/ushell/channels/flow/core/cmds/tips.py
+++ b/Engine/Extras/ushell/channels/flow/core/cmds/tips.py
@@ -6,6 +6,7 @@ import flow.cmd
 #-------------------------------------------------------------------------------
 class Tips(flow.cmd.Cmd):
     """ Displays a useful tip """
+    nocrlf = flow.cmd.Opt(False, "Do not output CRLF after tip line")
 
     @classmethod
     def read_get_tip_methods(self):
@@ -21,8 +22,9 @@ class Tips(flow.cmd.Cmd):
         with flow.cmd.text.light_cyan:
             print("TIP!", end="")
 
+        end = "" if self.args.nocrlf else "\n"
         with flow.cmd.text.white:
-            print("", tips[random.randrange(len(tips))])
+            print("", tips[random.randrange(len(tips))], end=end)
 
     def get_tips(self):
         return (
diff --git a/Engine/Extras/ushell/channels/flow/core/system/lib/bootstrap.py b/Engine/Extras/ushell/channels/flow/core/system/lib/bootstrap.py
--- a/Engine/Extras/ushell/channels/flow/core/system/lib/bootstrap.py
+++ b/Engine/Extras/ushell/channels/flow/core/system/lib/bootstrap.py
@@ -100,14 +100,14 @@ def _http_get(url, dest_dir, progress_cb=None):
 def _extract_tar(payload, dest_dir):
     cmd = ("tar", "-x", "-C", dest_dir, "-f", payload)
     _, ext = os.path.splitext(payload)
-    if ext == '.zip' and os.name == 'posix':
-        # Tar on linux cannot open .zip files 
+    if ext == ".zip" and os.name == "posix":
+        # Tar on Linux cannot open .zip files
         cmd = ("unzip", payload, '-d', dest_dir)
     ret = sp.run(cmd, stdout=sp.DEVNULL)
     assert ret.returncode == 0, "Extract failed; " + " ".join(cmd)
 
 #-------------------------------------------------------------------------------
-def _tool_extract_method(self, src_path, dest_dir):
+def _tool_extract_method(src_path, dest_dir):
     extractors = {
         ".tar.gz"   : _extract_tar,
         ".tgz"      : _extract_tar,
@@ -131,14 +131,12 @@ def _tool_extract_method(self, src_path, dest_dir):
     for file_name in os.listdir(src_path):
         os.rename(src_path + file_name, dest_dir + file_name)
 
-flow.describe.Tool.extract = _tool_extract_method
-
 
 
 #-------------------------------------------------------------------------------
 def _validate_tool(name, tool):
     assert tool._version, f"No version set on tool '{name}'"
-    assert tool._bundles, f"No payload set on tool '{name}'"
+    assert tool._bundles, f"No payloads added for tool '{name}'"
 
 #-------------------------------------------------------------------------------
 def _acquire_tool(name, tool, manifest, target_dir, progress_cb):
@@ -191,12 +189,13 @@ def _acquire_tool(name, tool, manifest, target_dir, progress_cb):
 
     # Extract the payloads.
     progress_cb("uzip")
+    extract_func = getattr(tool, "extract", _tool_extract_method)
     extract_dir = temp_dir + "$extract/"
     os.mkdir(extract_dir)
     for bundle_name, bundle in bundles.items():
-        _log.write("Extract", bundle_name, "from", fetch_dir, "to", extract_dir, "with", tool.extract)
+        _log.write("Extract", bundle_name, "from", fetch_dir, "to", extract_dir, "with", extract_func)
         fetch_dir = temp_dir + bundle_name + "/"
-        tool.extract(fetch_dir, extract_dir)
+        extract_func(fetch_dir, extract_dir)
 
     # Collapse
     source_dir = str(extract_dir)
@@ -616,7 +615,7 @@ class _Sources(object):
 def _get_sources() -> _Sources:
     sources = _Sources()
 
-    def _add_source(sources:_Sources, source_dir:Path, group_name:str="") -> int:
+    def _add_source(source_dir:Path, group_name:str="") -> int:
         count = 0
         for item in source_dir.glob("*"):
             if not item.is_dir():
@@ -627,25 +626,25 @@ def _get_sources() -> _Sources:
                 sources.add_source(group_name + item.name, candidate)
                 count += 1
             elif not group_name:
-                count += _add_source(sources, item, item.name + ".")
+                count += _add_source(item, item.name + ".")
 
         return count
 
     for item in Path(__file__).parents:
         if item.name == "channels":
-            _add_source(sources, item)
+            _add_source(item)
             break
 
     from_home = Path.home() / ".ushell/channels"
     if from_home.is_dir():
-        _add_source(sources, from_home, "$home.")
+        _add_source(from_home, "$home.")
 
     if from_env := os.getenv("FLOW_CHANNELS_DIR"):
         import shlex
         for item in shlex.split(from_env):
             item = Path(item).resolve()
             if item.is_dir():
-                _add_source(sources, item, "$env.")
+                _add_source(item, "$env.")
 
     return sources
 
diff --git a/Engine/Extras/ushell/channels/flow/nt/describe.flow.py b/Engine/Extras/ushell/channels/flow/nt/describe.flow.py
--- a/Engine/Extras/ushell/channels/flow/nt/describe.flow.py
+++ b/Engine/Extras/ushell/channels/flow/nt/describe.flow.py
@@ -23,12 +23,12 @@ fd.source("[REDACTED_PROFILE_URL]/fd/releases/latest", "fd-v([0-9.]+)-x86_64
 
 #-------------------------------------------------------------------------------
 fzf = flow.describe.Tool()
-fzf.version("0.54.3")
+fzf.version("0.56.3")
 fzf.source("[REDACTED_PROFILE_URL]/fzf/releases/latest", r"fzf-(\d+\.\d+\.\d+)-windows")
 if bundle := fzf.bundle(platform="win32"):
     bundle.payload("[REDACTED_PROFILE_URL]/fzf/releases/download/v$VERSION/fzf-$VERSION-windows_amd64.zip")
     bundle.bin("fzf.exe")
-    bundle.sha1("2c3ceab1a099ad3ab30eaffb6edaaa8f1a8ce2b4")
+    bundle.sha1("9dc22afb3a687b10eac57fe27f1a0e4d65c52944")
 
 #-------------------------------------------------------------------------------
 shell_cmd = flow.describe.Command()
diff --git a/Engine/Extras/ushell/channels/flow/nt/provision.bat b/Engine/Extras/ushell/channels/flow/nt/provision.bat
--- a/Engine/Extras/ushell/channels/flow/nt/provision.bat
+++ b/Engine/Extras/ushell/channels/flow/nt/provision.bat
@@ -23,6 +23,7 @@ if exist "%~f1\current\%_pymark%" 1>nul 2>nul (
 call:check_bin tar.exe _tar_path
 call:check_bin curl.exe _curl_path
 call:check_bin certutil.exe _certutil_path
+call:check_bin robocopy.exe _robocopy_path
 
 1>nul 2>nul (
     mkdir "%_tempdir%"
@@ -77,7 +78,7 @@ for /f "delims=, tokens=2" %%d in ('tasklist.exe /fo csv /fi "imagename eq cmd.e
     echo off
 )
 
-move "%_workdir%" "%~1" 1>nul 2>nul
+%_robocopy_path% "%_workdir%" "%~1" /MIR /MOVE 1>nul 2>nul
 if not exist "%~1" (
     call:on_error "Failed finalising Python provision"
 )
diff --git a/Engine/Extras/ushell/channels/flow/nt/shell_cmd.py b/Engine/Extras/ushell/channels/flow/nt/shell_cmd.py
--- a/Engine/Extras/ushell/channels/flow/nt/shell_cmd.py
+++ b/Engine/Extras/ushell/channels/flow/nt/shell_cmd.py
@@ -111,7 +111,7 @@ class _Shell(object):
             print(rf'doskey dp.="{thefuzz_path}" pushdir .', file=out)
 
             print("clink_x64.exe inject", file=out)
-            print("$tip.exe", file=out)
+            print("$tip.exe --nocrlf", file=out)
 
             if user_script:
                 print(rf'call "{user_script}"', file=out)
diff --git a/Engine/Extras/ushell/channels/unreal/core/boot.py b/Engine/Extras/ushell/channels/unreal/core/boot.py
--- a/Engine/Extras/ushell/channels/unreal/core/boot.py
+++ b/Engine/Extras/ushell/channels/unreal/core/boot.py
@@ -17,10 +17,8 @@ class Boot(flow.cmd.Cmd):
                 stdin=subprocess.DEVNULL
             )
             if proc.returncode:
-                self.print_error()
                 self.print_error("Set a shortcut's 'Start In' to a folder with a valid .uproject")
                 self.print_error("path or launch with the 'ushell --project=[uprojpath]' argument.")
-                self.print_error()
         except Exception as e:
             self.print_error(e)
         print()
diff --git a/Engine/Extras/ushell/channels/unreal/core/cmds/build.py b/Engine/Extras/ushell/channels/unreal/core/cmds/build.py
--- a/Engine/Extras/ushell/channels/unreal/core/cmds/build.py
+++ b/Engine/Extras/ushell/channels/unreal/core/cmds/build.py
@@ -553,10 +553,14 @@ class _Runtime(_BuildCmd, unrealcmd.MultiPlatformCmd):
 
         platform = self.get_platform(platform).get_name()
 
+        ue_context = self.get_unreal_context()
+        has_project = ue_context.get_project() is not None
+
         builder = self.get_builder()
         builder.set_platform(platform)
         builder.set_targets(target)
         builder.set_variant(self.args.variant)
+        builder.set_projected(has_project)
 
         return self.run_build(builder)
 
diff --git a/Engine/Extras/ushell/channels/unreal/core/cmds/sln.py b/Engine/Extras/ushell/channels/unreal/core/cmds/sln.py
--- a/Engine/Extras/ushell/channels/unreal/core/cmds/sln.py
+++ b/Engine/Extras/ushell/channels/unreal/core/cmds/sln.py
@@ -5,6 +5,32 @@ import shutil
 import unrealcmd
 from pathlib import Path
 
+#-------------------------------------------------------------------------------
+def _make_temp_vcxproj(ctx, path, name, target, ubt_args=""):
+    eng_dir = ctx.get_engine().get_dir()
+
+    import slnformer
+    p = slnformer._Project(name, None)
+    p.add_property("ConfigurationType", "Makefile")
+    p.add_property("ProjectGuid", p.get_guid())
+    p.add_config("Debug")
+    p.add_config("DebugGame")
+    p.add_config("Development")
+    p.add_config("Test")
+    p.add_config("Shipping")
+    p.add_property("bin_suffix", "-Win64-$(Configuration)", cond="$(Configuration) != 'Development'")
+    p.add_property("BuildDir", str(eng_dir / "Build\\BatchFiles"))
+    p.set_build_action(f"$(BuildDir)\\Build.bat {target} Win64 $(Configuration) {ubt_args}")
+    p.set_rebuild_action(f"$(BuildDir)\\Rebuild.bat {target} Win64 $(Configuration) {ubt_args}")
+    p.set_clean_action(f"$(BuildDir)\\Clean.bat {target} Win64 $(Configuration) {ubt_args}")
+    p.add_property("NMakeOutput", str(eng_dir / "Binaries\\Win64\\{target}$(bin_suffix).exe"))
+    p.add_import("$(VCTargetsPath)/Microsoft.Cpp.Default.props")
+    p.add_import("$(VCTargetsPath)/Microsoft.Cpp.props")
+    p.add_import("$(VCTargetsPath)/Microsoft.Cpp.targets")
+
+    with open(path, "wt") as out:
+        p.write(out)
+
 #-------------------------------------------------------------------------------
 class _Base(unrealcmd.MultiPlatformCmd):
     def _get_primary_name(self):
@@ -30,9 +56,18 @@ class _Base(unrealcmd.MultiPlatformCmd):
 
     def _get_sln_path(self):
         ue_context = self.get_unreal_context()
-        engine = ue_context.get_engine()
+
         primary_name = self._get_primary_name()
-        return engine.get_dir().parent / (primary_name + ".sln")
+        base_name = primary_name + ".sln"
+
+        sln_path = None
+        if project := ue_context.get_project():
+            sln_path = project.get_dir() / base_name
+        if sln_path is None or not sln_path.is_file():
+            engine = ue_context.get_engine()
+            sln_path = engine.get_dir().parent / base_name
+
+        return sln_path
 
     def _open_sln(self):
         if os.name != "nt":
@@ -106,6 +141,15 @@ class Generate(_Base):
             if ret:
                 return ret
 
+        if not self.args.all:
+            if project := ue_context.get_project():
+                if not (project.get_dir() / "Source").is_dir():
+                    int_dir = project.get_dir() / "Intermediate/ProjectFiles"
+                    ubt_arg = f"-Project=\"{project.get_path()}\""
+                    for target in ("Game", "Client", "Server", "Editor"):
+                        vcxproj_path = int_dir / f"ushell_{project.get_name()}{target}.vcxproj"
+                        _make_temp_vcxproj(ue_context, vcxproj_path, project.get_name(), "Unreal" + target, ubt_arg)
+
         if self.args.open == "open":
             return self._open_sln()
 
@@ -213,6 +257,8 @@ class Tiny(_Base):
         former = slnformer.Sln(self._get_sln_path(), proj_dir)
 
         tiny_proj = former.add_project("_tiny", makefile=False)
+        tiny_proj.add_config("Debug")
+        tiny_proj.add_config("Development")
         tiny_proj.add_property("ushell_bin_suffix", "-Win64-$(Configuration)", cond="$(Configuration) != 'Development'")
         tiny_proj.add_property("ushell_bin_prefix", "../../../")
         tiny_proj.add_property("ushell_bin", "$(ushell_bin_prefix)Binaries/Win64/$(ushell_target)$(ushell_bin_suffix).exe")
diff --git a/Engine/Extras/ushell/channels/unreal/core/describe.flow.py b/Engine/Extras/ushell/channels/unreal/core/describe.flow.py
--- a/Engine/Extras/ushell/channels/unreal/core/describe.flow.py
+++ b/Engine/Extras/ushell/channels/unreal/core/describe.flow.py
@@ -4,19 +4,19 @@ import flow.describe
 
 #-------------------------------------------------------------------------------
 fzf = flow.describe.Tool()
-fzf.version("0.54.3")
+fzf.version("0.56.3")
 if bundle := fzf.bundle(platform="linux"):
     bundle.payload("[REDACTED_PROFILE_URL]/fzf/releases/download/v$VERSION/fzf-$VERSION-linux_amd64.tar.gz")
     bundle.bin("fzf")
-    bundle.sha1("e46fcd29c453ee61be0d41ca22b7eefd944f91f6")
+    bundle.sha1("e8cc2ff14b0a39d1ba0c87e8f350557dcc56ac51")
 if bundle := fzf.bundle(platform="darwin-amd64"):
     bundle.payload("[REDACTED_PROFILE_URL]/fzf/releases/download/v$VERSION/fzf-$VERSION-darwin_amd64.tar.gz")
     bundle.bin("fzf")
-    bundle.sha1("3723b3445037bdaa302d1d365d4940beb4c283c8")
+    bundle.sha1("40a234fa002071aab34e24b86a43bf93cf99e067")
 if bundle := fzf.bundle(platform="darwin-arm64"):
     bundle.payload("[REDACTED_PROFILE_URL]/fzf/releases/download/v$VERSION/fzf-$VERSION-darwin_arm64.tar.gz")
     bundle.bin("fzf")
-    bundle.sha1("c52657e62c095909fb1ac69e6d714eac3a10dc4c")
+    bundle.sha1("59c4cb882cab45ee58586870a1a74022da43400c")
 
 
 
diff --git a/Engine/Extras/ushell/channels/unreal/core/pylib/slnformer.py b/Engine/Extras/ushell/channels/unreal/core/pylib/slnformer.py
--- a/Engine/Extras/ushell/channels/unreal/core/pylib/slnformer.py
+++ b/Engine/Extras/ushell/channels/unreal/core/pylib/slnformer.py
@@ -75,6 +75,7 @@ class _Project(_Node):
         super().__init__(name, parent)
         self._props = []
         self._imports = { "pre" : [], "post" : [] }
+        self._configs = set()
 
     def set_build_action(self, cmd:str, *, cond:str="") -> None:
         return self.add_property("NMakeBuildCommandLine", cmd, cond=cond)
@@ -91,6 +92,9 @@ class _Project(_Node):
     def add_import(self, imp_path:Path, *, pre:bool=False) -> None:
         self._imports["pre" if pre else "post"].append(str(imp_path))
 
+    def add_config(self, name:str) -> None:
+        self._configs.add(name)
+
     def write(self, out:IO) -> None:
         cond_props = {}
         for k,v,c in self._props:
@@ -101,8 +105,8 @@ class _Project(_Node):
         with xmler.Project(DefaultTargets="Build", xmlns="http://schemas.microsoft.com/developer/msbuild/2003") as node:
             # TODO: formalise this!
             node._write( " <ItemGroup>")
-            node._write(r'  <ProjectConfiguration Include="Debug|x64"><Configuration>Debug</Configuration><Platform>x64</Platform></ProjectConfiguration>')
-            node._write(r'  <ProjectConfiguration Include="Development|x64"><Configuration>Development</Configuration><Platform>x64</Platform></ProjectConfiguration>')
+            for config in self._configs:
+                node._write(fr'  <ProjectConfiguration Include="{config}|x64"><Configuration>{config}</Configuration><Platform>x64</Platform></ProjectConfiguration>')
             node._write( " </ItemGroup>")
 
             for imp in self._imports["pre"]:
diff --git a/Engine/Extras/ushell/channels/unreal/core/pylib/unreal/_context.py b/Engine/Extras/ushell/channels/unreal/core/pylib/unreal/_context.py
--- a/Engine/Extras/ushell/channels/unreal/core/pylib/unreal/_context.py
+++ b/Engine/Extras/ushell/channels/unreal/core/pylib/unreal/_context.py
@@ -425,7 +425,10 @@ class _Project(_UfsNode):
             return
 
         if not (self.get_dir() / "Source").is_dir():
-            return self.get_parent_node().get_target_by_type(target_type)
+            ret = self.get_parent_node().get_target_by_type(target_type)
+            if ret:
+                ret._ufs_mount = self.get_root_mount()
+            return ret
 
         for target in self._read_targets("Source", ancestors=False, shallow=True):
             if target.get_type() == target_type:
