Updated ushell [38941985]

diff --git a/Engine/Extras/ushell/channels/unreal/core/cmds/zen.py b/Engine/Extras/ushell/channels/unreal/core/cmds/zen.py
--- a/Engine/Extras/ushell/channels/unreal/core/cmds/zen.py
+++ b/Engine/Extras/ushell/channels/unreal/core/cmds/zen.py
@@ -239,6 +239,15 @@ class ImportSnapshot(unreal.zencmd.ZenUtilityBaseCmd):
             args.append(snapshot['bucket'])
             args.append('--key')
             args.append(snapshot['key'])
+        elif snapshot_type == 'builds':
+            args.append('--builds')
+            args.append(sourcehost)
+            args.append('--namespace')
+            args.append(snapshot['namespace'])
+            args.append('--bucket')
+            args.append(snapshot['bucket'])
+            args.append('--builds-id')
+            args.append(snapshot['builds-id'])
         elif snapshot_type == 'zen':
             args.append('--zen')
             args.append(sourcehost)
