Updated ushell [39059900]

diff --git a/Engine/Extras/ushell/channels/flow/nt/boot.ps1 b/Engine/Extras/ushell/channels/flow/nt/boot.ps1
--- a/Engine/Extras/ushell/channels/flow/nt/boot.ps1
+++ b/Engine/Extras/ushell/channels/flow/nt/boot.ps1
@@ -1,35 +1,28 @@
-Write-Debug "Ushell bootstrap args 2: $Args"
+# Copyright Epic Games, Inc. All Rights Reserved.
 
-$Working = Join-Path $env:LOCALAPPDATA "ushell" ".working"
-if(![String]::IsNullOrEmpty($env:flow_working_dir)) {
+$Working = [IO.Path]::Combine($env:LOCALAPPDATA, "ushell\.working")
+if (![String]::IsNullOrEmpty($env:flow_working_dir)) {
     $Working = $env:flow_working_dir
 }
 
-Write-Debug "Working dir: $Working"
-
-$ProvisionBat = Join-Path $PSScriptRoot "provision.bat"
-Write-Debug "Provisioning: $ProvisionBat"
+$ProvisionBat = [IO.Path]::Combine($PSScriptRoot, "provision.bat")
 & "cmd.exe" "/d/c" "`"$ProvisionBat`" $Working"
 
-if(!$?) {
+if (!$?) {
     throw "Ushell provisioning failed"
 }
 
 
-$PythonPath = Join-Path $Working "python\current\flow_python.exe"
-$Bootpy = Join-Path $PSScriptRoot "..\core\system\boot.py"
-Write-Debug "Python: $PythonPath"
+$PythonPath = [IO.Path]::Combine($Working, "python\current\flow_python.exe")
+$Bootpy = [IO.Path]::Combine($PSScriptRoot, "..\core\system\boot.py")
 
-$TempDir = Join-Path Temp: "ushell"
-New-Item -ItemType Directory -Force $TempDir
-$Cookie = Join-Path $TempDir "cmd_boot_$(New-Guid).ps1"
-$Cookie = New-Item $Cookie
-Write-Debug "Cookie: $Cookie"
+$TempDir = [IO.Path]::Combine($env:TEMP, "ushell")
 
+$Cookie = [IO.Path]::Combine($TempDir, "pwsh_boot_$(New-Guid).ps1")
 & $PythonPath "-Xutf8" "-Esu" $Bootpy "--bootarg=pwsh,$Cookie" $Args | Out-Host
 $PythonReturn = $?
 
-if(!$PythonReturn) {
+if (!$PythonReturn) {
     throw "Python boot failed"
 }
 
@@ -37,9 +30,7 @@ try {
     Import-Module $Cookie
 }
 finally {
-    if( Test-Path $Cookie )
-    {
-        Write-Debug "Removing cookie: $Cookie"
+    if (Test-Path $Cookie) {
         Remove-Item $Cookie
     }
 }
\ No newline at end of file
diff --git a/Engine/Extras/ushell/channels/flow/nt/shell_pwsh.py b/Engine/Extras/ushell/channels/flow/nt/shell_pwsh.py
--- a/Engine/Extras/ushell/channels/flow/nt/shell_pwsh.py
+++ b/Engine/Extras/ushell/channels/flow/nt/shell_pwsh.py
@@ -24,7 +24,7 @@ class _Shell(object):
         try: os.makedirs(os.path.dirname(cookie))
         except: pass
 
-        system = self.get_system();
+        system = self.get_system()
         working_dir = system.get_working_dir()
         shims_path = os.path.normpath(working_dir) + "/shims"
         cmd_tree = system.get_command_tree()
@@ -43,6 +43,8 @@ class _Shell(object):
             # Environment must be written out first so autocomplete daemon gets the right session id
             for key, value in env.read_changes():
                 value = value or ""
+                value = value.replace("\"", "`\"")
+                value = value.replace("$", "`$")
                 out.write(f'${{env:{key}}}="{value}"\n')
             out.write(header.format(shims_path=shims_path,working_dir=working_dir))
             out.write(autocompleter.format(python=sys.executable, run_py_path=run_py_path, working_dir=working_dir))
@@ -66,19 +68,20 @@ Register-ArgumentCompleter -CommandName "{name}" -Native -ScriptBlock $SharedCom
 # Explicitly export this function so that others are not exported automatically
 Export-ModuleMember -Function {name}
 """
+
 # If we give commands long names, we need to separate completers that provide the name ushell $complete expects
 def _get_autocompleter_script():
     return r"""
 function Get-Daemon {{
-    if( $script:CompleteDaemonCache -ne $null -and !$script:CompleteDaemonCache.HasExited) {{
+    if ($script:CompleteDaemonCache -ne $null -and !$script:CompleteDaemonCache.HasExited) {{
         return $script:CompleteDaemonCache
     }}
 
     $ProcStart = New-Object System.Diagnostics.ProcessStartInfo
     # the $ in $complete must be escaped for powershell
     $DaemonExe = $true
-    if( $DaemonExe ) {{
-        $ProcStart.FileName = Join-Path $ShimsPath "`$complete"
+    if ($DaemonExe) {{
+        $ProcStart.FileName = [IO.Path]::Combine($ShimsPath, "`$complete")
         $ProcStart.Arguments = @("--daemon")
     }}
     else {{
@@ -86,18 +89,18 @@ function Get-Daemon {{
         $ProcStart.Arguments = @(
             "-Xutf8",
             "-Esu",
-            "{run_py_path}",
-            (Join-Path {working_dir} "manifest"),
+            "`"{run_py_path}`"",
+            "`"{working_dir}/manifest`"",
             "`$complete",
             "--daemon"
-            )
+        )
     }}
     $ProcStart.UseShellExecute = $false
     $ProcStart.RedirectStandardOutput = $true
     $ProcStart.RedirectStandardInput = $true
 
     $script:CompleteDaemonCache = [System.Diagnostics.Process]::Start($ProcStart)
-    if( $null -eq $script:CompleteDaemonCache ) {{
+    if ($null -eq $script:CompleteDaemonCache) {{
         throw "Failed to start autocomplete Daemon!"
     }}
     return $script:CompleteDaemonCache
@@ -110,7 +113,7 @@ $SharedCompleter = {{
     $Words[0] = $Words[0] -replace "\.exe"
 
     # Add ... if we have a partial argument
-    if(![string]::IsNullOrEmpty($WordToComplete)) {{
+    if (![string]::IsNullOrEmpty($WordToComplete)) {{
         $Words[-1] += "..."
     }}
     else {{
@@ -119,23 +122,24 @@ $SharedCompleter = {{
 
     $CompleteDaemon = Get-Daemon
 
-    $CompleteDaemon.StandardInput.WriteLine("`u{{01}}")
-    foreach( $Word in $Words ) {{
+    $CompleteDaemon.StandardInput.WriteLine("$([char]0x01)")
+    foreach ($Word in $Words) {{
         $CompleteDaemon.StandardInput.WriteLine($Word)
     }}
-    $CompleteDaemon.StandardInput.WriteLine("`u{{02}}")
+    $CompleteDaemon.StandardInput.WriteLine("$([char]0x02)")
+    $CompleteDaemon.StandardInput.Flush()
 
     $Results = @()
     $Continue = $true
-    while($Continue) {{
+    while ($Continue) {{
         $Option = $CompleteDaemon.StandardOutput.Readline();
         # -eq and .Equals behave differently for the ascii separator bytes used here
         switch -exact ($Option) {{
-            {{ "`u{{01}}".Equals($Option) }} {{
+            {{ "$([char]0x01)".Equals($Option) }} {{
                 $Continue = $false
                 break;
             }}
-            {{ "`u{{02}}".Equals($Option) }} {{
+            {{ "$([char]0x02)".Equals($Option) }} {{
                 $Results = @();
                 $Continue = $false
                 break;
@@ -154,7 +158,7 @@ $SharedCompleter = {{
     # Filter to only arguments that start with our partial match if it exists
     $Results = $Results | Where-Object {{ $_ -like "$WordToComplete*" }}
 
-    if( $results.Count -eq 0 ) {{
+    if ($results.Count -eq 0) {{
         # prevent powershell from autocompleting paths to match ushell behavior in cmd
         ""
     }}
@@ -173,35 +177,43 @@ function Update-UShellEnvVars {{
 
     $CompleteDaemon = Get-Daemon
 
-    if( $null -eq $CompleteDaemon ) {{
+    if ($null -eq $CompleteDaemon) {{
         Write-Error "Update-UShellEnvVars - CompleteDaemon is null!"
     }}
-    if( $CompleteDaemon.HasExited ) {{
+    if ($CompleteDaemon.HasExited) {{
         Write-Error "Update-UShellEnvVars - CompleteDaemon has exited!"
     }}
 
-    $CompleteDaemon.StandardInput.WriteLine("`u{{01}}")
+    $CompleteDaemon.StandardInput.WriteLine("$([char]0x01)")
     $CompleteDaemon.StandardInput.WriteLine("`$`$")
     $CompleteDaemon.StandardInput.WriteLine($($executionContext.SessionState.Path.CurrentLocation))
-    $CompleteDaemon.StandardInput.WriteLine("`u{{02}}")
+    $CompleteDaemon.StandardInput.WriteLine("$([char]0x02)")
 
     $Results = @{{}}
     $Continue = $true
-    while($Continue) {{
+    while ($Continue) {{
         $Key = $CompleteDaemon.StandardOutput.Readline();
         # -eq and .Equals behave differently for the ascii separator bytes used here
-        if( [string]::IsNullOrEmpty($Key) -or $Key.Equals("`u{{01}}") ) {{ break; }}                 #EOF
-        if( $Key.Equals("`u{{02}}") ) {{ $Results = @{{}}; break; }} #error
+        if ([string]::IsNullOrEmpty($Key) -or $Key.Equals("$([char]0x01)")) {{
+            break
+        }}
 
-        $Value = $CompleteDaemon.StandardOutput.Readline();
-        if( [string]::IsNullOrEmpty($Value) -or $Value.Equals("`u{{01}}") -or $Value.Equals("`u{{02}}") ) {{ $Results = @{{}}; break; }} #error
+        if ($Key.Equals("$([char]0x02)")) {{
+            $Results = @{{}}
+            break
+        }}
+
+        $Value = $CompleteDaemon.StandardOutput.Readline()
+        if ([string]::IsNullOrEmpty($Value) -or $Value.Equals("$([char]0x01)") -or $Value.Equals("$([char]0x02)")) {{
+            $Results = @{{}}
+            break
+        }}
 
         $Results[$Key] = $Value
     }}
 
-    foreach( $i in $Results.GetEnumerator()) {{
-        $EnvPath = (Join-Path Env: "USHELL_$($i.Key)")
-        Write-Debug "Setting env var: $EnvPath"
+    foreach ($i in $Results.GetEnumerator()) {{
+        $EnvPath = "Env:\USHELL_$($i.Key)"
         $null = $script:EnvVarsToRemove.Add($EnvPath)
         New-Item $EnvPath -Value $i.Value -Force
     }}
@@ -214,12 +226,10 @@ Export-ModuleMember -Function "Update-UShellEnvVars"
 def _get_cleanup_script():
     return r"""
 $ExecutionContext.SessionState.Module.OnRemove += {
-    if( $script:CompleteDaemonCache -ne $null ) {
-        Write-Debug "Stopping complete daemon"
+    if ($script:CompleteDaemonCache -ne $null) {
         Stop-Process $script:CompleteDaemonCache
     }
-    foreach( $v in $script:EnvVarsToRemove.GetEnumerator()) {
-        Write-Debug "Removing env var: $v"
+    foreach ($v in $script:EnvVarsToRemove.GetEnumerator()) {
         Remove-Item $v
     }
 }
diff --git a/Engine/Extras/ushell/channels/unreal/core/cmds/zen.py b/Engine/Extras/ushell/channels/unreal/core/cmds/zen.py
--- a/Engine/Extras/ushell/channels/unreal/core/cmds/zen.py
+++ b/Engine/Extras/ushell/channels/unreal/core/cmds/zen.py
@@ -225,7 +225,7 @@ class ImportSnapshot(unreal.zencmd.ZenUtilityBaseCmd):
         snapshot_type = snapshot['type']
 
         sourcehost = None
-        if snapshot_type == 'cloud' or snapshot_type == 'zen':
+        if snapshot_type == 'cloud' or snapshot_type == 'zen' or snapshot_type == 'builds':
             sourcehost = snapshot['host']
             if self.args.sourcehost:
                 sourcehost = self.conform_hostname(self.args.sourcehost)
