Updated ushell [39207369]

#rnx

diff --git a/Engine/Extras/ushell/channels/unreal/core/cmds/zen.py b/Engine/Extras/ushell/channels/unreal/core/cmds/zen.py
--- a/Engine/Extras/ushell/channels/unreal/core/cmds/zen.py
+++ b/Engine/Extras/ushell/channels/unreal/core/cmds/zen.py
@@ -329,7 +329,8 @@ class ImportSnapshot(unreal.zencmd.ZenUtilityBaseCmd):
         if not self.add_args_from_descriptor(snapshot, args):
             return False
 
-        if snapshot['type'] == 'cloud':
+        snapshot_type = snapshot['type']
+        if snapshot_type == 'cloud' or snapshot_type == 'builds':
             if not self.refresh_zen_token():
                 return False
 
diff --git a/Engine/Extras/ushell/channels/unreal/core/pylib/unreal/_platform.py b/Engine/Extras/ushell/channels/unreal/core/pylib/unreal/_platform.py
--- a/Engine/Extras/ushell/channels/unreal/core/pylib/unreal/_platform.py
+++ b/Engine/Extras/ushell/channels/unreal/core/pylib/unreal/_platform.py
@@ -79,8 +79,8 @@ class Platform(object):
             # and C-style comments. So we'll try and conform it
             ret = ""
             for line in inp:
-                try: line = line[:line.index("//")]
-                except ValueError: pass
+                if m := re.match(r'^(.*[^"])//', line):
+                    line = m.group(1)
                 line = line.strip()
                 if not line:
                     continue
diff --git a/Engine/Extras/ushell/channels/unreal/perforce/cmds/clean.py b/Engine/Extras/ushell/channels/unreal/perforce/cmds/clean.py
--- a/Engine/Extras/ushell/channels/unreal/perforce/cmds/clean.py
+++ b/Engine/Extras/ushell/channels/unreal/perforce/cmds/clean.py
@@ -79,11 +79,9 @@ class _Cleaner(object):
 
         if os.name == "nt":
             sp.Popen(
-                ("cmd.exe", "/c", "rd", "/q/s", str(self._rubbish_dir)),
+                ("cmd.exe", "/d/c", "rd", "/q/s", str(self._rubbish_dir)),
                 stdout=sp.DEVNULL,
-                stderr=sp.DEVNULL
-                )
-            pass
+                stderr=sp.DEVNULL)
         else:
             if self._rubbish_dir.is_dir() and self._rubbish_dir != "/": # paranoia!
                 sp.Popen(
