Updated ushell [39254281]

#rnx

diff --git a/Engine/Extras/ushell/channels/unreal/core/pylib/unreal/_context.py b/Engine/Extras/ushell/channels/unreal/core/pylib/unreal/_context.py
--- a/Engine/Extras/ushell/channels/unreal/core/pylib/unreal/_context.py
+++ b/Engine/Extras/ushell/channels/unreal/core/pylib/unreal/_context.py
@@ -439,11 +439,7 @@ class _Project(_UfsNode):
         if target_type != TargetType.GAME:
             target_name += type_name
 
-        target = _Target(target_name, target_type, self.get_root_mount())
-        if create_temp and Platform.get_host() != "Mac":
-            self._create_temp_target(target)
-
-        return target
+        return _Target(target_name, target_type, self.get_root_mount())
 
 
 
diff --git a/Engine/Extras/ushell/channels/unreal/core/pylib/unreal/platforms/mac/ios.py b/Engine/Extras/ushell/channels/unreal/core/pylib/unreal/platforms/mac/ios.py
--- a/Engine/Extras/ushell/channels/unreal/core/pylib/unreal/platforms/mac/ios.py
+++ b/Engine/Extras/ushell/channels/unreal/core/pylib/unreal/platforms/mac/ios.py
@@ -1,12 +1,27 @@
 # Copyright Epic Games, Inc. All Rights Reserved.
 
 import os
+import sys
+
 import unreal
+import plistlib
+import subprocess
 
 #-------------------------------------------------------------------------------
 class Platform(unreal.Platform):
     name = "IOS"
 
+    def get_connected_devices(self):
+        engine_dir = self.get_unreal_context().get_engine().get_dir()
+        ideviceid_file = os.path.join(engine_dir, "Extras/ThirdPartyNotUE/libimobiledevice/Mac/idevice_id")
+        if not os.path.isfile(ideviceid_file):
+            raise EnvironmentError(f"Failed to find idevice_id: '{ideviceid_file}'")
+        try:
+            result = subprocess.check_output([ideviceid_file, "-l"], encoding="utf-8").strip()
+            return result
+        except subprocess.CalledProcessError as e:
+            print(f"Failed to run command: '{e}'")
+
     def _read_env(self):
         yield from ()
 
@@ -23,3 +38,46 @@ class Platform(unreal.Platform):
     def _get_cook_form(self, target):
         if target == "game":   return "IOS"
         if target == "client": return "IOSClient"
+
+    def _launch(self, exec_context, stage_dir, binary_path, args):
+        if Platform.get_host() != "Mac":
+            raise EnvironmentError(f"Launch of iOS apps is currently only supported on Macs")
+
+        if stage_dir:
+            # Check that staged dir exists
+            if not os.path.isdir(stage_dir):
+                raise EnvironmentError(f"App is not staged: {stage_dir}'")
+
+            # Check that the Info.plist file exists
+            plist_file = binary_path + ".app/Info.plist"
+            if not os.path.isfile(plist_file):
+                raise EnvironmentError(f"Failed to find Info.plist file: '{plist_file}'")
+
+            # find the CFBundleIdentifier
+            bundle_id = ""
+            with open(plist_file, 'rb') as fp:
+                pl = plistlib.load(fp)
+                bundle_id = pl["CFBundleIdentifier"]
+                print(f"App's BundleIdentifier: '{bundle_id}'")
+
+            # Is there one (and only one) iOS device connected?
+            device_udid = self.get_connected_devices()
+            if not device_udid:
+                raise EnvironmentError("Failed to find any connected iOS devices")
+
+            if len(device_udid.split()) != 1:
+                raise EnvironmentError("There is more than one connected iOS device")
+
+            args = "devicectl device process launch".split()
+
+            args = (*args, "--terminate-existing")
+            args = (*args, "--console")
+
+            args = (*args, "--device")
+            args = (*args, device_udid)
+
+            args = (*args, bundle_id)
+
+        cmd = exec_context.create_runnable("xcrun", *args)
+        cmd.launch()
+        return True
