Updated ushell [39713409]

#rnx

diff --git a/Engine/Extras/ushell/channels/unreal/core/cmds/zen.py b/Engine/Extras/ushell/channels/unreal/core/cmds/zen.py
--- a/Engine/Extras/ushell/channels/unreal/core/cmds/zen.py
+++ b/Engine/Extras/ushell/channels/unreal/core/cmds/zen.py
@@ -10,56 +10,6 @@ import unreal.zencmd
 import http.client
 import json
 import hashlib
-import tempfile
-
-#-------------------------------------------------------------------------------
-def _get_ips_socket():
-    import socket
-    def _get_ip(family, ip_addr):
-        with socket.socket(family, socket.SOCK_DGRAM) as s:
-            try:
-                s.connect((ip_addr, 1))
-                return s.getsockname()[0]
-            except Exception as e:
-                return
-    ret = (_get_ip(socket.AF_INET6, "fe80::0"), _get_ip(socket.AF_INET, "172.31.255.255"))
-    return [x for x in ret if x]
-
-def _get_ips_osx():
-    ret = _get_ips_socket()
-    try:
-        import subprocess
-        import re
-        segments = re.split("\\s+", subprocess.check_output('ifconfig | grep 169.254.', shell=True).decode('utf-8').strip())
-        print(segments)
-        ret.append(segments[1])
-    except:
-        pass
-    return ret
-
-def _get_ips_wmic():
-    import subprocess as sp
-    ret = []
-    current_ips = None
-    wmic_cmd = ("wmic", "nicconfig", "where", "ipenabled=true", "get", "ipaddress,ipconnectionmetric", "/format:value")
-    with sp.Popen(wmic_cmd, stdout=sp.PIPE, stderr=sp.DEVNULL) as proc:
-        for line in proc.stdout:
-            if line.startswith(b"IPAddress="):
-                if line := line[10:].strip():
-                    current_ips = line
-            elif line.startswith(b"IPConnectionMetric="):
-                if line := line[19:].strip():
-                    connection_metric = int(line)
-                    for ip in eval(current_ips):
-                        ret.append((connection_metric, ip))
-    ret.sort(key=lambda p: p[0])
-    ret = [item[1] for item in ret]
-    ret.sort(key=lambda x: "." in x) # sort ipv4 last
-    return ret
-
-_get_ips = _get_ips_wmic if os.name == "nt" else (_get_ips_osx if sys.platform == "darwin" else _get_ips_socket)
-
-
 
 #-------------------------------------------------------------------------------
 class Dashboard(unreal.zencmd.ZenUETargetBaseCmd):
@@ -114,7 +64,6 @@ class ImportSnapshot(unreal.zencmd.ZenUtilityBaseCmd):
     projectid = unrealcmd.Opt("", "The zen project ID to import into (defaults to an ID based on the current ushell project)")
     oplog = unrealcmd.Opt("", "The zen oplog to import into (defaults to the oplog name in the snapshot)")
     sourcehost = unrealcmd.Opt("", "The source host to import from (defaults to the host specified in the snapshot)")
-    cloudauthservice = unrealcmd.Opt("", "Name of the service to authorize with when importing from a cloud source")
     asyncimport = unrealcmd.Opt(False, "Trigger import but don't wait for completion")
     forceimport = unrealcmd.Opt(False, "Force import of all attachments")
 
@@ -135,221 +84,6 @@ class ImportSnapshot(unreal.zencmd.ZenUtilityBaseCmd):
             if item.is_dir():
                 yield item.name + "/"
 
-    def _lookup_service_name(self):
-        if getattr(self, "_service_name", None):
-            return
-
-        config = self.get_unreal_context().get_config()
-        default = config.get("Engine", "StorageServers", "Default")
-        if value := default.OAuthProviderIdentifier:
-            self._service_name = value
-        else:
-            cloud = config.get("Engine", "StorageServers", "Cloud")
-            if value := cloud.OAuthProviderIdentifier:
-                self._service_name = value
-
-    def refresh_zen_token(self):
-        ue_context = self.get_unreal_context()
-        engine = ue_context.get_engine()
-
-        bin_type = "win-x64"
-        if sys.platform == "darwin": bin_type = "osx-x64"
-        elif sys.platform == "linux": bin_type = "linux-x64"
-
-        bin_ext = ".exe"
-        if sys.platform == "darwin": bin_ext = ""
-        elif sys.platform == "linux": bin_ext = ""
-
-        bin_path = engine.get_dir()
-        bin_path /= f"Binaries/DotNET/OidcToken/{bin_type}/OidcToken{bin_ext}"
-        if not bin_path.is_file():
-            raise FileNotFoundError(f"Unable to find '{bin_path}'")
-
-        self._service_name = self.args.cloudauthservice
-        self._lookup_service_name()
-        if not self._service_name:
-            raise ValueError("Unable to discover service name")
-
-        token_dir = tempfile.TemporaryDirectory()
-        token_file_path = f"{token_dir.name}/oidctoken.json"
-        oidcargs = (
-            "--ResultToConsole=true",
-            "--Service=" + str(self._service_name),
-            f"--OutFile={token_file_path}"
-        )
-
-        if project := ue_context.get_project():
-            oidcargs = (*oidcargs, "--Project=" + str(project.get_path()))
-
-        cmd = self.get_exec_context().create_runnable(bin_path, *oidcargs)
-        ret, output = cmd.run2()
-        if ret != 0:
-            return False
-
-        tokenresponse = None
-        with open(token_file_path, "r") as token_file:
-            tokenresponse = json.load(token_file)
-
-        self._AccessToken = tokenresponse['Token']
-        self.print_info("Cloud access token obtained for import operation.")
-
-        return True
-
-    def get_exec_context(self):
-        context = super().get_exec_context()
-        if hasattr(self, '_AccessToken') and self._AccessToken:
-            if sys.platform == 'win32':
-                context.get_env()["UE-CloudDataCacheAccessToken"] = self._AccessToken
-            else:
-                context.get_env()["UE_CloudDataCacheAccessToken"] = self._AccessToken
-        return context
-
-    def conform_to_cook_forms(self, oplogid):
-        ue_context = self.get_unreal_context()
-        platforms = ue_context.get_platform_provider()
-        for platform_name in platforms.read_platform_names():
-            platform = platforms.get_platform(platform_name)
-            for form_name in ("client", "game", "server"):
-                try:
-                    cook_form = platform.get_cook_form(form_name)
-                    if cook_form.lower() == oplogid.lower():
-                        return cook_form
-                except:
-                    pass
-        return oplogid
-
-    def conform_hostname(self, hostname):
-        return hostname
-
-    def add_args_from_descriptor(self, snapshot, args):
-        snapshot_type = snapshot['type']
-
-        sourcehost = None
-        if snapshot_type == 'cloud' or snapshot_type == 'zen' or snapshot_type == 'builds':
-            sourcehost = snapshot['host']
-            if self.args.sourcehost:
-                sourcehost = self.conform_hostname(self.args.sourcehost)
-
-        if snapshot_type == 'cloud':
-            args.append('--cloud')
-            args.append(sourcehost)
-            args.append('--namespace')
-            args.append(snapshot['namespace'])
-            args.append('--bucket')
-            args.append(snapshot['bucket'])
-            args.append('--key')
-            args.append(snapshot['key'])
-        elif snapshot_type == 'builds':
-            args.append('--builds')
-            args.append(sourcehost)
-            args.append('--namespace')
-            args.append(snapshot['namespace'])
-            args.append('--bucket')
-            args.append(snapshot['bucket'])
-            args.append('--builds-id')
-            args.append(snapshot['builds-id'])
-        elif snapshot_type == 'zen':
-            args.append('--zen')
-            args.append(sourcehost)
-            args.append('--source-project')
-            args.append(snapshot['projectid'])
-            args.append('--source-oplog')
-            args.append(snapshot['oplogid'])
-        elif snapshot_type == 'file':
-            args.append('--file')
-            args.append(snapshot['directory'])
-            args.append('--name')
-            args.append(snapshot['filename'])
-        else:
-            self.print_error(f"Unsupported snapshot type {snapshot_type}")
-            return False
-
-        return True
-
-    def perform_target_creation(self, snapshot):
-        target_projectid = self.get_project_or_default(self.args.projectid)
-        if self.find_project_by_id(target_projectid) is None:
-            ue_context = self.get_unreal_context()
-            args = [
-            "project-create",
-            target_projectid,
-            ue_context.get_branch().get_dir(),
-            ue_context.get_engine().get_dir(),
-            ue_context.get_project().get_dir(),
-            ue_context.get_project().get_path()
-            ]
-            self.run_zen_utility(args)
-
-        target_oplogid = self.args.oplog or snapshot['targetplatform']
-
-        ue_context = self.get_unreal_context()
-        target_gcmarker = ue_context.get_project().get_dir() / f"Saved/Cooked/{self.conform_to_cook_forms(target_oplogid)}/ue.projectstore"
-
-        target_gcmarker_dir = os.path.dirname(target_gcmarker)
-        if not os.path.exists(target_gcmarker_dir):
-            os.makedirs(target_gcmarker_dir)
-        projectstore = {
-            "zenserver" : {
-                "islocalhost" : True,
-                "hostname" : "[::1]",
-                "remotehostnames" : _get_ips(),
-                "hostport" : 8558,
-                "projectid" : target_projectid,
-                "oplogid" : target_oplogid
-            }
-        }
-        with open(target_gcmarker, "w") as target_gcmarker_outfile:
-            json.dump(projectstore, target_gcmarker_outfile, indent='\t')
-
-        args = [
-        "oplog-create",
-        target_projectid,
-        target_oplogid,
-        target_gcmarker,
-        "--force-update"
-        ]
-        self.run_zen_utility(args)
-        return True
-
-    def perform_import(self, snapshot):
-        args = [
-        "oplog-import",
-        self.get_project_or_default(self.args.projectid),
-        self.args.oplog or snapshot['targetplatform'],
-        "--ignore-missing-attachments",
-        "--clean"
-        ]
-
-        if self.args.asyncimport:
-            args.append("--async")
-
-        if self.args.forceimport:
-            args.append("--force")
-
-        if not self.add_args_from_descriptor(snapshot, args):
-            return False
-
-        snapshot_type = snapshot['type']
-        if snapshot_type == 'cloud' or snapshot_type == 'builds':
-            if not self.refresh_zen_token():
-                return False
-
-        return self.run_zen_utility(args)
-
-    def perform_clean_post_import(self, snapshot):
-        target_oplogid = self.args.oplog or snapshot['targetplatform']
-
-        ue_context = self.get_unreal_context()
-        target_dir = ue_context.get_project().get_dir() / f"Saved/Cooked/{self.conform_to_cook_forms(target_oplogid)}"
-
-        for fsentry in os.scandir(target_dir):
-            if fsentry.is_dir(follow_symlinks=False):
-                shutil.rmtree(fsentry.path)
-            else:
-                if fsentry.name != 'ue.projectstore':
-                    os.unlink(fsentry.path)
-
-
     @unrealcmd.Cmd.summarise
     def main(self):
         self.launch_zenserver_if_not_running()
@@ -362,12 +96,4 @@ class ImportSnapshot(unreal.zencmd.ZenUtilityBaseCmd):
 
         snapshot = descriptor["snapshots"][self.args.snapshotindex]
 
-        if not self.perform_target_creation(snapshot):
-            self.print_error(f"Error creating import target location")
-            return False
-
-        import_result = self.perform_import(snapshot)
-        if import_result == 0:
-            self.perform_clean_post_import(snapshot)
-        return import_result
-
+        return self.import_snapshot(snapshot, self.args.sourcehost, self.args.projectid, self.args.oplogid, self.args.asyncimport, self.args.forceimport)
diff --git a/Engine/Extras/ushell/channels/unreal/core/pylib/unreal/zencmd.py b/Engine/Extras/ushell/channels/unreal/core/pylib/unreal/zencmd.py
--- a/Engine/Extras/ushell/channels/unreal/core/pylib/unreal/zencmd.py
+++ b/Engine/Extras/ushell/channels/unreal/core/pylib/unreal/zencmd.py
@@ -10,9 +10,58 @@ import http.client
 import json
 import hashlib
 import socket
+import tempfile
+
+#-------------------------------------------------------------------------------
+def _get_ips_socket():
+    import socket
+    def _get_ip(family, ip_addr):
+        with socket.socket(family, socket.SOCK_DGRAM) as s:
+            try:
+                s.connect((ip_addr, 1))
+                return s.getsockname()[0]
+            except Exception as e:
+                return
+    ret = (_get_ip(socket.AF_INET6, "fe80::0"), _get_ip(socket.AF_INET, "172.31.255.255"))
+    return [x for x in ret if x]
+
+def _get_ips_osx():
+    ret = _get_ips_socket()
+    try:
+        import subprocess
+        import re
+        segments = re.split("\\s+", subprocess.check_output('ifconfig | grep 169.254.', shell=True).decode('utf-8').strip())
+        print(segments)
+        ret.append(segments[1])
+    except:
+        pass
+    return ret
+
+def _get_ips_wmic():
+    import subprocess as sp
+    ret = []
+    current_ips = None
+    wmic_cmd = ("wmic", "nicconfig", "where", "ipenabled=true", "get", "ipaddress,ipconnectionmetric", "/format:value")
+    with sp.Popen(wmic_cmd, stdout=sp.PIPE, stderr=sp.DEVNULL) as proc:
+        for line in proc.stdout:
+            if line.startswith(b"IPAddress="):
+                if line := line[10:].strip():
+                    current_ips = line
+            elif line.startswith(b"IPConnectionMetric="):
+                if line := line[19:].strip():
+                    connection_metric = int(line)
+                    for ip in eval(current_ips):
+                        ret.append((connection_metric, ip))
+    ret.sort(key=lambda p: p[0])
+    ret = [item[1] for item in ret]
+    ret.sort(key=lambda x: "." in x) # sort ipv4 last
+    return ret
+
+_get_ips = _get_ips_wmic if os.name == "nt" else (_get_ips_osx if sys.platform == "darwin" else _get_ips_socket)
 
 #-------------------------------------------------------------------------------
 class ZenBaseCmd(unrealcmd.Cmd):
+    cloudauthservice = unrealcmd.Opt("", "Name of the service to authorize with when importing from a cloud source")
 
     def _build(self, target, variant, platform):
         if target.get_type() == unreal.TargetType.EDITOR:
@@ -162,6 +211,222 @@ class ZenBaseCmd(unrealcmd.Cmd):
         ue_context = self.get_unreal_context()
         return self.get_project_id_from_project(ue_context.get_project())
 
+    def _lookup_service_name(self):
+        if getattr(self, "_service_name", None):
+            return
+
+        config = self.get_unreal_context().get_config()
+        default = config.get("Engine", "StorageServers", "Default")
+        if value := default.OAuthProviderIdentifier:
+            self._service_name = value
+        else:
+            cloud = config.get("Engine", "StorageServers", "Cloud")
+            if value := cloud.OAuthProviderIdentifier:
+                self._service_name = value
+
+    def get_exec_context(self):
+        context = super().get_exec_context()
+        if hasattr(self, '_AccessToken') and self._AccessToken:
+            if sys.platform == 'win32':
+                context.get_env()["UE-CloudDataCacheAccessToken"] = self._AccessToken
+            else:
+                context.get_env()["UE_CloudDataCacheAccessToken"] = self._AccessToken
+        return context
+
+    def refresh_zen_token(self):
+        ue_context = self.get_unreal_context()
+        engine = ue_context.get_engine()
+
+        bin_type = "win-x64"
+        if sys.platform == "darwin": bin_type = "osx-x64"
+        elif sys.platform == "linux": bin_type = "linux-x64"
+
+        bin_ext = ".exe"
+        if sys.platform == "darwin": bin_ext = ""
+        elif sys.platform == "linux": bin_ext = ""
+
+        bin_path = engine.get_dir()
+        bin_path /= f"Binaries/DotNET/OidcToken/{bin_type}/OidcToken{bin_ext}"
+        if not bin_path.is_file():
+            raise FileNotFoundError(f"Unable to find '{bin_path}'")
+
+        self._service_name = self.args.cloudauthservice
+        self._lookup_service_name()
+        if not self._service_name:
+            raise ValueError("Unable to discover service name")
+
+        token_dir = tempfile.TemporaryDirectory()
+        token_file_path = f"{token_dir.name}/oidctoken.json"
+        oidcargs = (
+            "--ResultToConsole=true",
+            "--Service=" + str(self._service_name),
+            f"--OutFile={token_file_path}"
+        )
+
+        if project := ue_context.get_project():
+            oidcargs = (*oidcargs, "--Project=" + str(project.get_path()))
+
+        cmd = self.get_exec_context().create_runnable(bin_path, *oidcargs)
+        ret, output = cmd.run2()
+        if ret != 0:
+            return False
+
+        tokenresponse = None
+        with open(token_file_path, "r") as token_file:
+            tokenresponse = json.load(token_file)
+
+        self._AccessToken = tokenresponse['Token']
+        self.print_info("Cloud access token obtained for import operation.")
+
+        return True
+
+    def conform_to_cook_forms(self, oplogid):
+        ue_context = self.get_unreal_context()
+        platforms = ue_context.get_platform_provider()
+        for platform_name in platforms.read_platform_names():
+            platform = platforms.get_platform(platform_name)
+            for form_name in ("client", "game", "server"):
+                try:
+                    cook_form = platform.get_cook_form(form_name)
+                    if cook_form.lower() == oplogid.lower():
+                        return cook_form
+                except:
+                    pass
+        return oplogid
+
+    def perform_target_creation(self, snapshot, projectid = None, oplogid = None):
+        target_projectid = self.get_project_or_default(projectid)
+        if self.find_project_by_id(target_projectid) is None:
+            ue_context = self.get_unreal_context()
+            args = [
+            "project-create",
+            target_projectid,
+            ue_context.get_branch().get_dir(),
+            ue_context.get_engine().get_dir(),
+            ue_context.get_project().get_dir(),
+            ue_context.get_project().get_path()
+            ]
+            self.run_zen_utility(args)
+
+        target_oplogid = oplogid or snapshot['targetplatform']
+
+        ue_context = self.get_unreal_context()
+        target_gcmarker = ue_context.get_project().get_dir() / f"Saved/Cooked/{self.conform_to_cook_forms(target_oplogid)}/ue.projectstore"
+
+        target_gcmarker_dir = os.path.dirname(target_gcmarker)
+        if not os.path.exists(target_gcmarker_dir):
+            os.makedirs(target_gcmarker_dir)
+        projectstore = {
+            "zenserver" : {
+                "islocalhost" : True,
+                "hostname" : "[::1]",
+                "remotehostnames" : _get_ips(),
+                "hostport" : 8558,
+                "projectid" : target_projectid,
+                "oplogid" : target_oplogid
+            }
+        }
+        with open(target_gcmarker, "w") as target_gcmarker_outfile:
+            json.dump(projectstore, target_gcmarker_outfile, indent='\t')
+
+        args = [
+        "oplog-create",
+        target_projectid,
+        target_oplogid,
+        target_gcmarker,
+        "--force-update"
+        ]
+        self.run_zen_utility(args)
+        return True
+
+    def add_args_from_descriptor(self, snapshot, sourcehost, args):
+        snapshot_type = snapshot['type']
+
+        if snapshot_type in ['cloud', 'zen', 'builds'] and not sourcehost:
+            sourcehost = snapshot['host']
+
+        if snapshot_type == 'cloud':
+            args.append('--cloud')
+            args.append(sourcehost)
+            args.append('--namespace')
+            args.append(snapshot['namespace'])
+            args.append('--bucket')
+            args.append(snapshot['bucket'])
+            args.append('--key')
+            args.append(snapshot['key'])
+        elif snapshot_type == 'builds':
+            args.append('--builds')
+            args.append(sourcehost)
+            args.append('--namespace')
+            args.append(snapshot['namespace'])
+            args.append('--bucket')
+            args.append(snapshot['bucket'])
+            args.append('--builds-id')
+            args.append(snapshot['builds-id'])
+        elif snapshot_type == 'zen':
+            args.append('--zen')
+            args.append(sourcehost)
+            args.append('--source-project')
+            args.append(snapshot['projectid'])
+            args.append('--source-oplog')
+            args.append(snapshot['oplogid'])
+        elif snapshot_type == 'file':
+            args.append('--file')
+            args.append(snapshot['directory'])
+            args.append('--name')
+            args.append(snapshot['filename'])
+        else:
+            self.print_error(f"Unsupported snapshot type {snapshot_type}")
+            return False
+
+        return True
+
+    def perform_import(self, snapshot, sourcehost, projectid, oplogid, asyncimport, forceimport):
+        args = [
+        "oplog-import",
+        self.get_project_or_default(projectid),
+        oplogid or snapshot['targetplatform'],
+        "--ignore-missing-attachments",
+        "--clean"
+        ]
+
+        if asyncimport:
+            args.append("--async")
+
+        if forceimport:
+            args.append("--force")
+
+        if not self.add_args_from_descriptor(snapshot, sourcehost, args):
+            return False
+
+        if snapshot['type'] in ['cloud', 'builds']:
+            if not self.refresh_zen_token():
+                return False
+
+        return self.run_zen_utility(args)
+
+    def perform_clean_post_import(self, snapshot, oplogid):
+        target_oplogid = oplogid or snapshot['targetplatform']
+
+        ue_context = self.get_unreal_context()
+        target_dir = ue_context.get_project().get_dir() / f"Saved/Cooked/{self.conform_to_cook_forms(target_oplogid)}"
+
+        for fsentry in os.scandir(target_dir):
+            if fsentry.is_dir(follow_symlinks=False):
+                shutil.rmtree(fsentry.path)
+            else:
+                if fsentry.name != 'ue.projectstore':
+                    os.unlink(fsentry.path)
+
+    def import_snapshot(self, snapshotdescriptor, sourcehost, projectid, oplogid, asyncimport, forceimport):
+        if not self.perform_target_creation(snapshotdescriptor, projectid, oplogid):
+            self.print_error(f"Error creating import target location")
+            return False
+
+        import_result = self.perform_import(snapshotdescriptor, sourcehost, projectid, oplogid, asyncimport, forceimport)
+        if import_result == 0:
+            self.perform_clean_post_import(snapshotdescriptor, oplogid)
+        return import_result
 
 #-------------------------------------------------------------------------------
 class ZenUETargetBaseCmd(ZenBaseCmd):
