Updated ushell [39985534]

#rnx

diff --git a/Engine/Extras/ushell/channels/flow/core/cmds/debug.py b/Engine/Extras/ushell/channels/flow/core/cmds/debug.py
--- a/Engine/Extras/ushell/channels/flow/core/cmds/debug.py
+++ b/Engine/Extras/ushell/channels/flow/core/cmds/debug.py
@@ -17,11 +17,15 @@ def _http_get(url, on_data):
     try:
         import certifi
         cafile_path = certifi.where()
+
+        import ssl
+        ssl_context = ssl.SSLContext()
+        ssl_context.load_verify_locations(cafile=cafile_path)
     except ImportError:
-        cafile_path = None
+        ssl_context = None
 
     # Fire up an http client to get the url.
-    client = urlopen(url, cafile=cafile_path)
+    client = urlopen(url, context=ssl_context)
     if (client.status // 100) != 2:
         assert False, f"Error creating HTTP client for {url}"
 
diff --git a/Engine/Extras/ushell/channels/flow/core/system/lib/bootstrap.py b/Engine/Extras/ushell/channels/flow/core/system/lib/bootstrap.py
--- a/Engine/Extras/ushell/channels/flow/core/system/lib/bootstrap.py
+++ b/Engine/Extras/ushell/channels/flow/core/system/lib/bootstrap.py
@@ -56,11 +56,15 @@ def _http_get(url, dest_dir, progress_cb=None):
     try:
         import certifi
         cafile_path = certifi.where()
+
+        import ssl
+        ssl_context = ssl.SSLContext()
+        ssl_context.load_verify_locations(cafile=cafile_path)
     except ImportError:
-        cafile_path = None
+        ssl_context = None
 
     # Fire up an http client to get the url.
-    client = urlopen(url, cafile=cafile_path)
+    client = urlopen(url, context=ssl_context)
     if (client.status // 100) != 2:
         assert False, f"Error creating HTTP client for {url}"
 
diff --git a/Engine/Extras/ushell/channels/unreal/core/cmds/build.py b/Engine/Extras/ushell/channels/unreal/core/cmds/build.py
--- a/Engine/Extras/ushell/channels/unreal/core/cmds/build.py
+++ b/Engine/Extras/ushell/channels/unreal/core/cmds/build.py
@@ -41,7 +41,7 @@ class _Builder(object):
         self._ue_context = ue_context
         self._targets = ()
         self._projected = False
-        self.set_platform(unreal.Platform.get_host())
+        self._platform = None
         self.set_variant("development")
 
     def set_targets(self, *targets):    self._targets = targets
@@ -53,10 +53,20 @@ class _Builder(object):
         self._args += list(args)
 
     def read_actions(self):
-        platform = self._platform
         variant = self._variant
-
         args = self._args
+
+        platform = self._platform
+        if not platform:
+            platform = unreal.Platform.get_host()
+
+        """
+        if platform.lower() == unreal.Platform.get_host().lower():
+            from platform import machine as get_machine
+            if get_machine() in ("arm64", "ARM64"):
+                args = (*args, "-Architecture=arm64")
+        """
+
         if len(self._targets) > 1:
             targets = []
             for target in self._targets:
diff --git a/Engine/Extras/ushell/channels/unreal/core/cmds/zen.py b/Engine/Extras/ushell/channels/unreal/core/cmds/zen.py
--- a/Engine/Extras/ushell/channels/unreal/core/cmds/zen.py
+++ b/Engine/Extras/ushell/channels/unreal/core/cmds/zen.py
@@ -96,4 +96,4 @@ class ImportSnapshot(unreal.zencmd.ZenUtilityBaseCmd):
 
         snapshot = descriptor["snapshots"][self.args.snapshotindex]
 
-        return self.import_snapshot(snapshot, self.args.sourcehost, self.args.projectid, self.args.oplogid, self.args.asyncimport, self.args.forceimport)
+        return self.import_snapshot(snapshot, self.args.sourcehost, self.args.projectid, self.args.oplog, self.args.asyncimport, self.args.forceimport)
diff --git a/Engine/Extras/ushell/channels/unreal/core/pylib/unreal/zencmd.py b/Engine/Extras/ushell/channels/unreal/core/pylib/unreal/zencmd.py
--- a/Engine/Extras/ushell/channels/unreal/core/pylib/unreal/zencmd.py
+++ b/Engine/Extras/ushell/channels/unreal/core/pylib/unreal/zencmd.py
@@ -37,6 +37,20 @@ def _get_ips_osx():
         pass
     return ret
 
+def _get_ips_powershell():
+    import subprocess as sp
+    ps_cmd = (
+        "powershell.exe",
+        "-NoProfile",
+        "-NoLogo",
+        "-Command",
+        r'(Get-CimInstance -Query "SELECT * FROM Win32_NetworkAdapterConfiguration WHERE IPEnabled=true"|Sort -Property IPConnectionMetric).IPAddress',
+    )
+    with sp.Popen(ps_cmd, stdout=sp.PIPE, stderr=sp.DEVNULL, stdin=sp.DEVNULL) as proc:
+        ret = [x.strip().decode() for x in proc.stdout]
+        proc.stdout.close()
+    return ret if proc.returncode == 0 else None
+
 def _get_ips_wmic():
     import subprocess as sp
     ret = []
@@ -53,11 +67,19 @@ def _get_ips_wmic():
                     for ip in eval(current_ips):
                         ret.append((connection_metric, ip))
     ret.sort(key=lambda p: p[0])
-    ret = [item[1] for item in ret]
+    return [item[1] for item in ret]
+
+def _get_ips_nt():
+    try:
+        ret = _get_ips_powershell()
+        if not ret:
+            ret = _get_ips_wmic()
+    except:
+        ret = _get_ips_socket()
     ret.sort(key=lambda x: "." in x) # sort ipv4 last
     return ret
 
-_get_ips = _get_ips_wmic if os.name == "nt" else (_get_ips_osx if sys.platform == "darwin" else _get_ips_socket)
+_get_ips = _get_ips_nt if os.name == "nt" else (_get_ips_osx if sys.platform == "darwin" else _get_ips_socket)
 
 #-------------------------------------------------------------------------------
 class ZenBaseCmd(unrealcmd.Cmd):
