Updated ushell [40136468]

#rnx

diff --git a/Engine/Extras/ushell/channels/flow/core/system/version b/Engine/Extras/ushell/channels/flow/core/system/version
--- a/Engine/Extras/ushell/channels/flow/core/system/version
+++ b/Engine/Extras/ushell/channels/flow/core/system/version
@@ -1 +1 @@
-22
+23
diff --git a/Engine/Extras/ushell/channels/flow/nt/provision.bat b/Engine/Extras/ushell/channels/flow/nt/provision.bat
--- a/Engine/Extras/ushell/channels/flow/nt/provision.bat
+++ b/Engine/Extras/ushell/channels/flow/nt/provision.bat
@@ -7,9 +7,9 @@ goto:eof
 ::------------------------------------------------------------------------------
 :main
 
-set _pyver=3.12.4
-set _pysha=15fea3c9367653a85086fe37216b4d1a1c78688fa5e1587e1db0b0f658856564
-set _pytag=312
+set _pyver=3.13.1
+set _pysha=7b7923ff0183a8b8fca90f6047184b419b108cb437f75fc1c002f9d2f8bcec16
+set _pytag=313
 set _pymark=%_pyver%.version
 
 :: temp folder for detritus that's rmdir'd at the end
diff --git a/Engine/Extras/ushell/channels/unreal/core/pylib/unreal/zencmd.py b/Engine/Extras/ushell/channels/unreal/core/pylib/unreal/zencmd.py
--- a/Engine/Extras/ushell/channels/unreal/core/pylib/unreal/zencmd.py
+++ b/Engine/Extras/ushell/channels/unreal/core/pylib/unreal/zencmd.py
@@ -9,11 +9,10 @@ import unreal.cmdline
 import http.client
 import json
 import hashlib
-import socket
 import tempfile
 
 #-------------------------------------------------------------------------------
-def _get_ips_socket():
+def _get_ips_socket(_target_platform):
     import socket
     def _get_ip(family, ip_addr):
         with socket.socket(family, socket.SOCK_DGRAM) as s:
@@ -25,15 +24,26 @@ def _get_ips_socket():
     ret = (_get_ip(socket.AF_INET6, "fe80::0"), _get_ip(socket.AF_INET, "172.31.255.255"))
     return [x for x in ret if x]
 
-def _get_ips_osx():
-    ret = _get_ips_socket()
+def _get_ips_osx(target_platform):
+    ret = _get_ips_socket(target_platform)
     try:
         import subprocess
         import re
-        segments = re.split("\\s+", subprocess.check_output('ifconfig | grep 169.254.', shell=True).decode('utf-8').strip())
+        # since grep will return an exit code of 1 if not found, force it to return 0, otherwise it throws an exception
+        segments = re.split("\\s+", subprocess.check_output('ifconfig | grep 169.254. || true', shell=True).decode('utf-8').strip())
         print(segments)
-        ret.append(segments[1])
+        # if you have multiple devices plugged in, this will iterate over the segments (every 6th item is the IP)
+        # ie: ['inet', '169.254.12.24', 'netmask', '0xffff0000', 'broadcast', '169.254.255.255', 'inet', '169.254.110.30', 'netmask', '0xffff0000', 'broadcast', '169.254.255.255']
+        ret.extend(segments[1::6])
+
+        # Can not rely on a stable server link local address above as each time USB cable is plugged in,
+        # the IP changes. So use the mac server's hostname for the client to connect against as a fallback.
+        # Note: only do this for members of the iOS family - and there's no guarantee the mac has a bonjour name
+        if any (x in target_platform.lower() for x in ["ios", "tvos", "vision"]):
+            if bonjourName := subprocess.check_output('scutil --get ComputerName', shell=True).decode('utf-8').strip():
+                ret.append("macserver://" + bonjourName)
     except:
+        print("Exception thrown in _get_ips_osx")
         pass
     return ret
 
@@ -69,13 +79,13 @@ def _get_ips_wmic():
     ret.sort(key=lambda p: p[0])
     return [item[1] for item in ret]
 
-def _get_ips_nt():
+def _get_ips_nt(target_platform):
     try:
         ret = _get_ips_powershell()
         if not ret:
             ret = _get_ips_wmic()
     except:
-        ret = _get_ips_socket()
+        ret = _get_ips_socket(target_platform)
     ret.sort(key=lambda x: "." in x) # sort ipv4 last
     return ret
 
@@ -342,7 +352,7 @@ class ZenBaseCmd(unrealcmd.Cmd):
             "zenserver" : {
                 "islocalhost" : True,
                 "hostname" : "[::1]",
-                "remotehostnames" : _get_ips(),
+                "remotehostnames" : _get_ips(snapshot['targetplatform']),
                 "hostport" : 8558,
                 "projectid" : target_projectid,
                 "oplogid" : target_oplogid
diff --git a/Engine/Extras/ushell/channels/unreal/perforce/cmds/clean.py b/Engine/Extras/ushell/channels/unreal/perforce/cmds/clean.py
--- a/Engine/Extras/ushell/channels/unreal/perforce/cmds/clean.py
+++ b/Engine/Extras/ushell/channels/unreal/perforce/cmds/clean.py
@@ -150,7 +150,9 @@ class Clean(flow.cmd.Cmd):
 
     def _append_saved(self, collector, saved_dir):
         if self.args.allsaved:
-            collector.add_dir(saved_dir)
+            for item in saved_dir.glob("*"):
+                if item.is_dir():
+                    collector.add_dir(item)
             return
 
         for sub_dir in (x for x in saved_dir.glob("*") if x.is_dir()):
