Updated ushell [40244361]

#rnx
#submitttool skip

diff --git a/Engine/Extras/ushell/channels/unreal/core/cmds/zen.py b/Engine/Extras/ushell/channels/unreal/core/cmds/zen.py
--- a/Engine/Extras/ushell/channels/unreal/core/cmds/zen.py
+++ b/Engine/Extras/ushell/channels/unreal/core/cmds/zen.py
@@ -97,3 +97,43 @@ class ImportSnapshot(unreal.zencmd.ZenUtilityBaseCmd):
         snapshot = descriptor["snapshots"][self.args.snapshotindex]
 
         return self.import_snapshot(snapshot, self.args.sourcehost, self.args.projectid, self.args.oplog, self.args.asyncimport, self.args.forceimport)
+
+#-------------------------------------------------------------------------------
+class CreateWorkspace(unreal.zencmd.ZenUtilityBaseCmd):
+    """ Create a Zen workspace"""
+    base_dir = unrealcmd.Arg(str, "Base path for the Zen workspace")
+    dynamic = unrealcmd.Opt(False, "Should the zen server allow for dynamic creation of shares by clients in this workspace")
+
+    @unrealcmd.Cmd.summarise
+    def main(self):
+        self.launch_zenserver_if_not_running()
+        workspaces = self.get_workspaces()
+        workspace = self.get_workspace(workspaces, self.args.base_dir)
+        if workspace is None:
+            self.create_workspace(self.args.base_dir, self.args.dynamic)
+        else:
+            self.print_info("Workspace already exists at this location")
+
+#-------------------------------------------------------------------------------
+class CreateShare(unreal.zencmd.ZenUtilityBaseCmd):
+    """ Create a Zen workspace share"""
+    share_dir = unrealcmd.Arg(str, "Base path for the Zen share. Must be inside an already existing workspace")
+
+    @unrealcmd.Cmd.summarise
+    def main(self):
+        self.launch_zenserver_if_not_running()
+        workspaces = self.get_workspaces()
+        share_path = os.path.abspath(self.args.share_dir)
+        found_workspace = None
+        for workspace in workspaces:
+            workspace_abs_path = os.path.abspath(workspace['root_path'])
+            if (os.path.commonpath([workspace_abs_path]) == os.path.commonpath([workspace_abs_path, share_path])):
+                found_workspace = workspace
+                share_path = os.path.relpath(share_path, workspace_abs_path)
+                break
+
+        if found_workspace is None:
+            self.print_error("No workspace found at " + share_path)
+            return False
+
+        self.create_workspaceshare(workspace['root_path'], share_path, False)
diff --git a/Engine/Extras/ushell/channels/unreal/core/describe.flow.py b/Engine/Extras/ushell/channels/unreal/core/describe.flow.py
--- a/Engine/Extras/ushell/channels/unreal/core/describe.flow.py
+++ b/Engine/Extras/ushell/channels/unreal/core/describe.flow.py
@@ -169,6 +169,14 @@ zen_importsnapshot = flow.describe.Command()
 zen_importsnapshot.source("cmds/zen.py", "ImportSnapshot")
 zen_importsnapshot.invoke("zen", "importsnapshot")
 
+zen_createworkspace = flow.describe.Command()
+zen_createworkspace.source("cmds/zen.py", "CreateWorkspace")
+zen_createworkspace.invoke("zen", "createworkspace")
+
+zen_createshare = flow.describe.Command()
+zen_createshare.source("cmds/zen.py", "CreateShare")
+zen_createshare.invoke("zen", "createshare")
+
 #-------------------------------------------------------------------------------
 sln_generate = flow.describe.Command()
 sln_generate.source("cmds/sln.py", "Generate")
diff --git a/Engine/Extras/ushell/channels/unreal/core/pylib/unreal/zencmd.py b/Engine/Extras/ushell/channels/unreal/core/pylib/unreal/zencmd.py
--- a/Engine/Extras/ushell/channels/unreal/core/pylib/unreal/zencmd.py
+++ b/Engine/Extras/ushell/channels/unreal/core/pylib/unreal/zencmd.py
@@ -56,7 +56,10 @@ def _get_ips_powershell():
         "-Command",
         r'(Get-CimInstance -Query "SELECT * FROM Win32_NetworkAdapterConfiguration WHERE IPEnabled=true"|Sort -Property IPConnectionMetric).IPAddress',
     )
-    with sp.Popen(ps_cmd, stdout=sp.PIPE, stderr=sp.DEVNULL, stdin=sp.DEVNULL) as proc:
+    startupinfo = sp.STARTUPINFO()
+    startupinfo.dwFlags = sp.STARTF_USESHOWWINDOW
+    startupinfo.wShowWindow = sp.SW_HIDE
+    with sp.Popen(ps_cmd, stdout=sp.PIPE, stderr=sp.DEVNULL, stdin=sp.DEVNULL, startupinfo=startupinfo, creationflags=sp.CREATE_NEW_CONSOLE) as proc:
         ret = [x.strip().decode() for x in proc.stdout]
         proc.stdout.close()
     return ret if proc.returncode == 0 else None
@@ -460,6 +463,88 @@ class ZenBaseCmd(unrealcmd.Cmd):
             self.perform_clean_post_import(snapshotdescriptor, oplogid)
         return import_result
 
+    def get_workspaces(self):
+        conn = http.client.HTTPConnection(host="localhost", port=8558, timeout=0.1)
+        headers = {"Accept":"application/json"}
+        workspaces_obj = {}
+        try:
+            conn.request('GET', '/ws', headers=headers)
+            response = conn.getresponse()
+            workspaces_obj = json.loads(response.read().decode())
+        except OSError:
+            return None
+
+        if workspaces_obj:
+            return workspaces_obj["workspaces"]
+        return None
+
+    def get_workspace(self, workspaces, base_dir):
+        for workspace in workspaces:
+            if workspace['root_path'] == base_dir:
+                return workspace
+        return None
+
+    def get_workspaceshare(self, workspace, share_dir):
+        for share in workspace.get('shares', []):
+            if share['share_path'] == share_dir:
+                return share
+        return None
+
+    def create_workspace(self, workspacepath, dynamic):
+        workspaces = self.get_workspaces()
+        workspace_exists = False
+        for workspace in workspaces:
+            if workspace['root_path'] == workspacepath:
+                workspace_exists = True
+                break
+
+        if not workspace_exists:
+            args = [
+                "workspace",
+                "create",
+                workspacepath,
+                "--allow-share-create-from-http" if dynamic else None
+            ]
+            if self.run_zen_utility(args) != 0:
+                return None
+
+        workspaces = self.get_workspaces()
+        for workspace in workspaces:
+            if workspace['root_path'] == workspacepath:
+                return workspace
+
+        return None
+
+    def create_workspaceshare(self, workspacepath, sharepath, dynamic):
+        existing_workspace = None
+        workspaces = self.get_workspaces()
+        for workspace in workspaces:
+            if workspace['root_path'] == workspacepath:
+                existing_workspace = workspace
+                break
+
+        if existing_workspace is None:
+            existing_workspace = self.create_workspace(workspacepath, dynamic)
+            if existing_workspace is None:
+                return False
+
+        existing_share = None
+        for share in existing_workspace.get('shares', []):
+            if share["share_path"] == sharepath:
+                existing_share = share
+                break
+
+        if existing_share is None:
+            args = [
+                "workspace-share",
+                "create",
+                existing_workspace["id"],
+                sharepath
+            ]
+            return self.run_zen_utility(args) == 0
+
+        return True
+
 #-------------------------------------------------------------------------------
 class ZenUETargetBaseCmd(ZenBaseCmd):
     variant = unrealcmd.Arg("development", "Build variant to launch")
