Updated ushell [40349509]

#rnx

diff --git a/Engine/Extras/ushell/channels/unreal/core/cmds/getbuild.py b/Engine/Extras/ushell/channels/unreal/core/cmds/getbuild.py
new file mode 100644
--- /dev/null
+++ b/Engine/Extras/ushell/channels/unreal/core/cmds/getbuild.py
@@ -0,0 +1,282 @@
+# Copyright Epic Games, Inc. All Rights Reserved.
+
+import sys
+import unrealcmd
+import tempfile
+import json
+import os
+from peafour import P4
+import urllib.request
+import string
+
+#-------------------------------------------------------------------------------
+class GetBuild(unrealcmd.Cmd):
+    """ Download a build/artifact from UE Cloud
+
+Unless otherwise specified it is downloaded to <projectdir>/Staged/<platform>.
+
+By default the latest build is downloaded. You can choose a specific build via --match as follows:
+    * --match 8v93f39cvn3           - If you know the build id:
+    * --match 1337                  - To search for the build cloest to a changelist
+    * --match job=v99f3;step=39f983 - For advanced searching you can specify key values"""
+
+    buildtype  = unrealcmd.Arg(str, "The type of build to download. Usually packaged or staged builds.")
+    platform  = unrealcmd.Arg(str, "The platform of the build you want to download ")
+    runargs  = unrealcmd.Arg([str], "Additional arguments to pass to the process being launched")
+
+    project = unrealcmd.Opt("", "The project you want to download a build for. Defaults to the current project in ushell.")
+    branch = unrealcmd.Opt("", "Branch/Stream to fetch data for, defaults to current p4 stream")
+    namespace = unrealcmd.Opt("", "The namespace to use overriding what is defined in engine ini")
+    host = unrealcmd.Opt("", "The host name of the server to connect to, defaults to your cloud storage server in engine ini")
+
+    dest = unrealcmd.Opt("", "The directory to download the build to. Local contents will be overwritten. Defaults to `<projectdir>/Staged/<platform>`")
+    match = unrealcmd.Opt("", "Either a build id (hexadecimal string), a changelist (number) or a key value seperate list of key=values you want to find a build for. e.g. changelist=1337;job=foo")
+
+    def lookup_service_name(self):
+        config = self.get_unreal_context().get_config()
+        default = config.get("Engine", "StorageServers", "Default")
+        if value := default.OAuthProviderIdentifier:
+            return value
+        else:
+            cloud = config.get("Engine", "StorageServers", "Cloud")
+            if value := cloud.OAuthProviderIdentifier:
+                return value
+
+    def lookup_namespace(self):
+        config = self.get_unreal_context().get_config()
+        default = config.get("Engine", "StorageServers", "Default")
+        if value := default.BuildsNamespace:
+            return value
+        else:
+            cloud = config.get("Engine", "StorageServers", "Cloud")
+            if value := cloud.BuildsNamespace:
+                return value
+
+    def lookup_host(self):
+        config = self.get_unreal_context().get_config()
+        default = config.get("Engine", "StorageServers", "Default")
+        host = ""
+        if value := default.Host:
+            host = str(value)
+        else:
+            cloud = config.get("Engine", "StorageServers", "Cloud")
+            if value := cloud.Host:
+                host=  str(value)
+
+        # if host is a semicolon seperated list we parse out the first entry as zen cli does not yet support a list of hosts
+        if host!= None and host.find(';') != -1:
+            host = host.split(';')[0]
+
+        return host
+
+    def refresh_access_token(self):
+        ue_context = self.get_unreal_context()
+        engine = ue_context.get_engine()
+
+        bin_type = "win-x64"
+        if sys.platform == "darwin": bin_type = "osx-x64"
+        elif sys.platform == "linux": bin_type = "linux-x64"
+
+        bin_ext = ".exe"
+        if sys.platform == "darwin": bin_ext = ""
+        elif sys.platform == "linux": bin_ext = ""
+
+        bin_path = engine.get_dir()
+        bin_path /= f"Binaries/DotNET/OidcToken/{bin_type}/OidcToken{bin_ext}"
+        if not bin_path.is_file():
+            raise FileNotFoundError(f"Unable to find '{bin_path}'")
+
+        service_name = self.lookup_service_name()
+        if not service_name:
+            raise ValueError("Unable to discover service name")
+
+        token_dir = tempfile.TemporaryDirectory()
+        token_file_path = f"{token_dir.name}/oidctoken.json"
+        oidcargs = (
+            "--ResultToConsole=true",
+            "--Service=" + str(service_name),
+            f"--OutFile={token_file_path}"
+        )
+
+        if project := ue_context.get_project():
+            oidcargs = (*oidcargs, "--Project=" + str(project.get_path()))
+
+        cmd = self.get_exec_context().create_runnable(bin_path, *oidcargs)
+        ret, output = cmd.run2()
+        if ret != 0:
+            return False
+
+        tokenresponse = None
+        with open(token_file_path, "r") as token_file:
+            tokenresponse = json.load(token_file)
+
+        self._AccessToken = tokenresponse['Token']
+        self.print_info("Cloud access token obtained.")
+
+        return True
+
+    def get_exec_context(self):
+        context = super().get_exec_context()
+        if hasattr(self, '_AccessToken') and self._AccessToken:
+            context.get_env()["CloudAccessToken"] = self._AccessToken
+        return context
+
+    def complete_buildtype(self, prefix):
+        return ["packaged", "staged"]
+
+    def complete_platform(self, prefix):
+        server_platforms = ("win64", "linux", "mac")
+        for platform in unrealcmd.Cmd.complete_platform(self, prefix):
+            yield platform
+            if platform in server_platforms:
+                yield platform + ".server"
+
+    def complete_dest(self, prefix):
+        prefix = prefix or "."
+        for item in os.scandir(prefix):
+            if item.is_dir():
+                yield item.name + "/"
+
+    def convert_platform(self, platform):
+        return self.get_platform(platform).get_config_name()
+
+    def convert_buildtype(self, build_type):
+        if build_type == "packaged":
+            return "packaged-build"
+        elif build_type == "staged":
+            return "staged-build"
+        else:
+            return build_type
+
+    def _get_p4_stream(self):
+        self.print_info("Checking Perforce stream")
+        if branch := self.args.branch:
+            print("User provided:", branch)
+            return branch
+
+        branch_dir = self.get_unreal_context().get_branch().get_dir()
+        info = P4(d=str(branch_dir)).info().run()
+
+        branch = getattr(info, "clientStream", None)
+        if not branch:
+            raise EnvironmentError("Unable to get Perforce client name")
+        print("Using stream", branch)
+        # massage the perforce stream name to match expectations, no initial / and - seperation between depot and stream
+        return branch.lstrip('/').replace('/', '-')
+
+    def get_build_id(self, build_id_arg):
+        if len(build_id_arg) == 24 and all(char in string.hexdigits for char in build_id_arg):
+            return build_id_arg
+        return None
+
+    def search_build(self, host, namespace, bucket, search_request):
+        search_changelist = None
+        if search_request.isnumeric():
+            search_changelist = search_request
+            print("Searching for build closest to changelist: " + str(search_changelist))
+            search_request = "" # we special case handle the search for changelists
+
+        # Run a search against the api to find builds that matches the request
+        url = f"{host}/api/v2/builds/{namespace}/{bucket}/search"
+        headers = {
+            "Authorization" : f"Bearer {self._AccessToken}",
+            "Content-Type" : "application/json"
+        }
+
+        query = {}
+        for statement in search_request.split(';'):
+            if statement.find('=') != -1:
+                k, v = statement.split('=')
+                query[k] = {"$eq": v}
+
+        if search_changelist:
+            query["commit"] = {"$lte": search_changelist}
+
+        # we only need the first build found so limiting results to 1
+        data = json.dumps({"query": query, "options": {"limit": 1}}).encode()
+        request = urllib.request.Request(url, data=data, headers=headers)
+        response = urllib.request.urlopen(request)
+        body = response.read().decode()
+
+        results = json.loads(body)["results"]
+        if len(results) == 0:
+            return (None, None, query)
+
+        return (results[0]["buildId"], results[0]["metadata"], query)
+
+    def main(self):
+        ue_context = self.get_unreal_context()
+        engine = ue_context.get_engine()
+        project = self.args.project
+        project_path = None
+        if not project:
+            project = ue_context.get_project().get_name()
+            project_path = os.path.dirname(ue_context.get_project().get_path())
+
+        platform = self.convert_platform(self.args.platform)
+        build_type = self.convert_buildtype(self.args.buildtype)
+        dest = self.args.dest
+        if dest == "" and project_path:
+            dest = os.path.join(project_path, self.args.buildtype, platform)
+
+        build_id = self.get_build_id(self.args.match)
+        branch = self._get_p4_stream()
+
+        if self.args.namespace:
+            namespace = self.args.namespace
+        else:
+            namespace = self.lookup_namespace()
+
+        if self.args.host:
+            host = self.args.host
+        else:
+            host = self.lookup_host()
+
+        if not self.refresh_access_token():
+            return False
+
+        bucket = f"{project}.{build_type}.{branch}.{platform}".lower()
+
+        if not host:
+            return False
+        if not namespace:
+            return False
+
+        if not build_id:
+            self.print_info("Enumerating builds")
+            # Search for build id
+            (build_id, metadata, query) = self.search_build(host, namespace, bucket, self.args.match)
+            if not build_id:
+                return f"Failed to find any build that matches query {query} in namespace {namespace} and bucket {bucket}"
+            else:
+                build_name = metadata['name']
+                print(f"Found build {build_id} - {build_name}")
+        access_token_env_var = "CloudAccessToken"
+
+        args = (
+            "builds",
+            "download",
+            "--url=" + str(host),
+            "--namespace=" + str(namespace),
+            "--bucket=" + bucket,
+            "--access-token-env=" + access_token_env_var,
+            dest,
+            build_id
+        )
+
+        bin_type = "Win64"
+        if sys.platform == "darwin": bin_type = "Mac"
+        elif sys.platform == "linux": bin_type = "Linux"
+
+        bin_ext = ".exe"
+        if sys.platform == "darwin": bin_ext = ""
+        elif sys.platform == "linux": bin_ext = ""
+
+        bin_path = engine.get_dir()
+        bin_path /= f"Binaries/{bin_type}/zen{bin_ext}"
+        if not bin_path.is_file():
+            raise FileNotFoundError(f"Unable to find '{bin_path}'")
+        args = (*args, *self.args.runargs)
+        cmd = self.get_exec_context().create_runnable(bin_path, *args)
+        result = cmd.run()
+        return result
diff --git a/Engine/Extras/ushell/channels/unreal/core/describe.flow.py b/Engine/Extras/ushell/channels/unreal/core/describe.flow.py
--- a/Engine/Extras/ushell/channels/unreal/core/describe.flow.py
+++ b/Engine/Extras/ushell/channels/unreal/core/describe.flow.py
@@ -144,6 +144,11 @@ ddc_auth = flow.describe.Command()
 ddc_auth.source("cmds/ddc.py", "Auth")
 ddc_auth.invoke("ddc", "auth")
 
+#-------------------------------------------------------------------------------
+get_cmd = flow.describe.Command()
+get_cmd.source("cmds/getbuild.py", "GetBuild")
+get_cmd.invoke("getbuild")
+
 #-------------------------------------------------------------------------------
 zen_start = flow.describe.Command()
 zen_start.source("cmds/zen.py", "Start")
diff --git a/Engine/Extras/ushell/channels/unreal/core/pylib/unreal/platforms/_android.py b/Engine/Extras/ushell/channels/unreal/core/pylib/unreal/platforms/_android.py
new file mode 100644
--- /dev/null
+++ b/Engine/Extras/ushell/channels/unreal/core/pylib/unreal/platforms/_android.py
@@ -0,0 +1,171 @@
+# Copyright Epic Games, Inc. All Rights Reserved.
+
+import os
+import re
+import shutil
+import unreal
+import subprocess as sp
+from pathlib import Path
+
+#-------------------------------------------------------------------------------
+class PlatformBase(unreal.Platform):
+    name = "Android"
+    config_section = "/Script/AndroidRuntimeSettings.AndroidRuntimeSettings"
+
+    def _read_env(self):
+        version = self.get_version()
+        prefix = f"Android/{version}/"
+        env_vars = {}
+
+        # sdk
+        env_vars["ANDROID_HOME"] = "android-sdk-windows/" if version == "-22" else ""
+
+        # jdk
+        if version <= "-22":
+            env_vars["JAVA_HOME"] = "jdk1.8.0_77/"
+        elif version <= "-25":
+            env_vars["JAVA_HOME"] = "jre/"
+        else:
+            env_vars["JAVA_HOME"] = "jbr/"
+
+        # ndk
+        if version <= "-22":
+            env_vars["NDKROOT"] = "android-ndk-r14b/"
+        else:
+            if sdks_root := os.getenv("UE_SDKS_ROOT"):
+                from pathlib import Path
+                ndk_dir = Path(sdks_root)
+                ndk_dir /= "Host" + unreal.Platform.get_host()
+                ndk_dir /= prefix
+                ndk_dir /= "ndk"
+                if ndk_dir.is_dir():
+                    ndk_ver = max(x.name for x in ndk_dir.glob("*") if x.is_dir())
+                    env_vars["NDKROOT"] = "ndk/" + ndk_ver
+            env_vars.setdefault("NDKROOT", "ndk/")
+
+        # dispatch
+        for env_var, template in env_vars.items():
+            value = os.getenv(env_var)
+            if not value:
+                value = prefix + template
+            yield env_var, value
+
+    def _get_version_ue4(self):
+        dot_cs = self.get_unreal_context().get_engine().get_dir()
+        dot_cs /= "Source/Programs/UnrealBuildTool/Platform/Android/UEBuildAndroid.cs"
+        try:
+            import re
+            with open(dot_cs, "rt") as cs_in:
+                cs_in.read(8192)
+                lines = iter(cs_in)
+                next(x for x in lines if "override string GetRequiredSDKString()" in x)
+                next(lines) # {
+                for i in range(5):
+                    line = next(lines)
+                    if m := re.search(r'return "(-\d+)"', line):
+                        return m.group(1)
+        except (StopIteration, FileNotFoundError):
+            pass
+
+    def _get_version_ue5(self):
+        func = "GetAutoSDKDirectoryForMainVersion"
+        dot_cs = "Source/Programs/UnrealBuildTool/Platform/Android/AndroidPlatformSDK"
+        version = self._get_version_helper_ue5(dot_cs + ".Versions.cs", func)
+        return version or self._get_version_helper_ue5(dot_cs + ".cs", func)
+
+    def _get_android_home(self):
+        out = next((v for k,v in self._read_env() if k == "ANDROID_HOME"), None)
+        if not out:
+            return
+
+        for prefix in ("", unreal.Platform.get_sdks_dir()):
+            candidate = Path(prefix) / out
+            if candidate.is_dir():
+                return candidate
+
+    def _get_adb(self):
+        home_dir = self._get_android_home()
+        return (home_dir / "platform-tools/adb") if home_dir else "adb"
+
+    def _get_aapt(self):
+        home_dir = self._get_android_home()
+        if not home_dir:
+            return "aapt"
+
+        try:
+            build_tools_dir = home_dir / "build-tools"
+            latest = max(x.name for x in build_tools_dir.glob("*"))
+            return build_tools_dir / latest / "aapt"
+        except ValueError:
+            return
+
+    def _get_cook_form(self, target):
+        if target == "game":   return "Android_ASTC"
+        if target == "client": return "Android_ASTCClient"
+
+    def _get_cook_flavor(self):
+        return "ASTC"
+
+    def _get_stage_name(self, target):
+        return "Android"
+
+    def _get_package_name(self):
+        config = self.get_unreal_context().get_config()
+        if ret := str(config.get("Engine", PlatformBase.config_section, "PackageName")):
+            if "[PROJECT]" in ret:
+                project_name = self.get_unreal_context().get_project().get_name()
+                project_name = project_name.replace("-", "_")
+                ret = ret.replace("[PROJECT]", project_name)
+            return ret
+        raise ValueError(f"Failed querying '{PlatformBase.config_section}/PackageName'")
+
+    def _launch(self, exec_context, stage_dir, binary_path, args):
+        print("!!\n!! NOTE: .run does not currently use locally-built executables\n!!\n")
+
+        package_name = self._get_package_name()
+        intent = package_name + "/com.epicgames.unreal.SplashActivity"
+        print("Package:", package_name)
+        print("Intent: ", intent)
+        print()
+
+        adb = self._get_adb()
+        reverse_port_mappings = [
+                    41899,  # Network file server, DEFAULT_TCP_FILE_SERVING_PORT in NetworkMessage.h
+                    1981,   # Unreal Insights data collection, TraceInsightsModule.cpp
+                    8558,   # Zen Store
+        ]
+        for reverse_port_mapping in reverse_port_mappings:
+            reverse_cmd = (str(adb), "reverse", "tcp:"+str(reverse_port_mapping), "tcp:"+str(reverse_port_mapping), ">", "NUL")
+            os.system(" ".join(reverse_cmd))
+
+        args = "'" + " ".join(x.replace("\"", "\\\"") for x in args) + "'"
+        cmd = (str(adb), "shell", "am", "start", "--es", "cmdline", args, intent)
+        print(*cmd)
+
+        # subprocess.run() messes with quotes on Windows so we'll use os.system
+        self._kill()
+        os.system(" ".join(cmd))
+
+        return True
+        """ -- not yet implemented properly --
+        # get device
+        device = None
+        for line in exec_context.create_runnable(adb, "-d", "devices"):
+            if line.endswith("device") and not device:
+                device = line.split()[0]
+        print(device)
+
+        # get pid
+        pid = -1
+        for line in exec_context.create_runnable(adb, "-d", "shell", "run-as", package_name, "ps"):
+            print(line)
+
+        return (device, pid)
+        """
+
+    def _kill(self, target=""):
+        package_name = self._get_package_name()
+        adb = self._get_adb()
+        cmd = (adb, "shell", "am", "force-stop", str(package_name))
+        print(*cmd)
+        sp.run(cmd)
diff --git a/Engine/Extras/ushell/channels/unreal/core/pylib/unreal/platforms/mac/android.py b/Engine/Extras/ushell/channels/unreal/core/pylib/unreal/platforms/mac/android.py
new file mode 100644
--- /dev/null
+++ b/Engine/Extras/ushell/channels/unreal/core/pylib/unreal/platforms/mac/android.py
@@ -0,0 +1,7 @@
+# Copyright Epic Games, Inc. All Rights Reserved.
+
+from .. import _android
+
+#-------------------------------------------------------------------------------
+class Platform(_android.PlatformBase):
+    pass
diff --git a/Engine/Extras/ushell/channels/unreal/core/pylib/unreal/platforms/win64/android.py b/Engine/Extras/ushell/channels/unreal/core/pylib/unreal/platforms/win64/android.py
--- a/Engine/Extras/ushell/channels/unreal/core/pylib/unreal/platforms/win64/android.py
+++ b/Engine/Extras/ushell/channels/unreal/core/pylib/unreal/platforms/win64/android.py
@@ -1,171 +1,7 @@
 # Copyright Epic Games, Inc. All Rights Reserved.
 
-import os
-import re
-import shutil
-import unreal
-import subprocess as sp
-from pathlib import Path
+from .. import _android
 
 #-------------------------------------------------------------------------------
-class Platform(unreal.Platform):
-    name = "Android"
-    config_section = "/Script/AndroidRuntimeSettings.AndroidRuntimeSettings"
-
-    def _read_env(self):
-        version = self.get_version()
-        prefix = f"Android/{version}/"
-        env_vars = {}
-
-        # sdk
-        env_vars["ANDROID_HOME"] = "android-sdk-windows/" if version == "-22" else ""
-
-        # jdk
-        if version <= "-22":
-            env_vars["JAVA_HOME"] = "jdk1.8.0_77/"
-        elif version <= "-25":
-            env_vars["JAVA_HOME"] = "jre/"
-        else:
-            env_vars["JAVA_HOME"] = "jbr/"
-
-        # ndk
-        if version <= "-22":
-            env_vars["NDKROOT"] = "android-ndk-r14b/"
-        else:
-            if sdks_root := os.getenv("UE_SDKS_ROOT"):
-                from pathlib import Path
-                ndk_dir = Path(sdks_root)
-                ndk_dir /= "Host" + unreal.Platform.get_host()
-                ndk_dir /= prefix
-                ndk_dir /= "ndk"
-                if ndk_dir.is_dir():
-                    ndk_ver = max(x.name for x in ndk_dir.glob("*") if x.is_dir())
-                    env_vars["NDKROOT"] = "ndk/" + ndk_ver
-            env_vars.setdefault("NDKROOT", "ndk/")
-
-        # dispatch
-        for env_var, template in env_vars.items():
-            value = os.getenv(env_var)
-            if not value:
-                value = prefix + template
-            yield env_var, value
-
-    def _get_version_ue4(self):
-        dot_cs = self.get_unreal_context().get_engine().get_dir()
-        dot_cs /= "Source/Programs/UnrealBuildTool/Platform/Android/UEBuildAndroid.cs"
-        try:
-            import re
-            with open(dot_cs, "rt") as cs_in:
-                cs_in.read(8192)
-                lines = iter(cs_in)
-                next(x for x in lines if "override string GetRequiredSDKString()" in x)
-                next(lines) # {
-                for i in range(5):
-                    line = next(lines)
-                    if m := re.search(r'return "(-\d+)"', line):
-                        return m.group(1)
-        except (StopIteration, FileNotFoundError):
-            pass
-
-    def _get_version_ue5(self):
-        func = "GetAutoSDKDirectoryForMainVersion"
-        dot_cs = "Source/Programs/UnrealBuildTool/Platform/Android/AndroidPlatformSDK"
-        version = self._get_version_helper_ue5(dot_cs + ".Versions.cs", func)
-        return version or self._get_version_helper_ue5(dot_cs + ".cs", func)
-
-    def _get_android_home(self):
-        out = next((v for k,v in self._read_env() if k == "ANDROID_HOME"), None)
-        if not out:
-            return
-
-        for prefix in ("", unreal.Platform.get_sdks_dir()):
-            candidate = Path(prefix) / out
-            if candidate.is_dir():
-                return candidate
-
-    def _get_adb(self):
-        home_dir = self._get_android_home()
-        return (home_dir / "platform-tools/adb") if home_dir else "adb"
-
-    def _get_aapt(self):
-        home_dir = self._get_android_home()
-        if not home_dir:
-            return "aapt"
-
-        try:
-            build_tools_dir = home_dir / "build-tools"
-            latest = max(x.name for x in build_tools_dir.glob("*"))
-            return build_tools_dir / latest / "aapt"
-        except ValueError:
-            return
-
-    def _get_cook_form(self, target):
-        if target == "game":   return "Android_ASTC"
-        if target == "client": return "Android_ASTCClient"
-
-    def _get_cook_flavor(self):
-        return "ASTC"
-
-    def _get_stage_name(self, target):
-        return "Android"
-
-    def _get_package_name(self):
-        config = self.get_unreal_context().get_config()
-        if ret := str(config.get("Engine", Platform.config_section, "PackageName")):
-            if "[PROJECT]" in ret:
-                project_name = self.get_unreal_context().get_project().get_name()
-                project_name = project_name.replace("-", "_")
-                ret = ret.replace("[PROJECT]", project_name)
-            return ret
-        raise ValueError(f"Failed querying '{Platform.config_section}/PackageName'")
-
-    def _launch(self, exec_context, stage_dir, binary_path, args):
-        print("!!\n!! NOTE: .run does not currently use locally-built executables\n!!\n")
-
-        package_name = self._get_package_name()
-        intent = package_name + "/com.epicgames.unreal.SplashActivity"
-        print("Package:", package_name)
-        print("Intent: ", intent)
-        print()
-
-        adb = self._get_adb()
-        reverse_port_mappings = [
-                    41899,  # Network file server, DEFAULT_TCP_FILE_SERVING_PORT in NetworkMessage.h
-                    1981,   # Unreal Insights data collection, TraceInsightsModule.cpp
-                    8558,   # Zen Store
-        ]
-        for reverse_port_mapping in reverse_port_mappings:
-            reverse_cmd = (str(adb), "reverse", "tcp:"+str(reverse_port_mapping), "tcp:"+str(reverse_port_mapping), ">", "NUL")
-            os.system(" ".join(reverse_cmd))
-
-        args = "'" + " ".join(x.replace("\"", "\\\"") for x in args) + "'"
-        cmd = (str(adb), "shell", "am", "start", "--es", "cmdline", args, intent)
-        print(*cmd)
-
-        # subprocess.run() messes with quotes on Windows so we'll use os.system
-        self._kill()
-        os.system(" ".join(cmd))
-
-        return True
-        """ -- not yet implemented properly --
-        # get device
-        device = None
-        for line in exec_context.create_runnable(adb, "-d", "devices"):
-            if line.endswith("device") and not device:
-                device = line.split()[0]
-        print(device)
-
-        # get pid
-        pid = -1
-        for line in exec_context.create_runnable(adb, "-d", "shell", "run-as", package_name, "ps"):
-            print(line)
-
-        return (device, pid)
-        """
-
-    def _kill(self, target=""):
-        package_name = self._get_package_name()
-        adb = self._get_adb()
-        cmd = (adb, "shell", "am", "force-stop", str(package_name))
-        print(*cmd)
-        sp.run(cmd)
+class Platform(_android.PlatformBase):
+    pass
