Updated ushell [40495196]

#rnx

diff --git a/Engine/Extras/ushell/channels/flow/core/cmds/debug.py b/Engine/Extras/ushell/channels/flow/core/cmds/debug.py
--- a/Engine/Extras/ushell/channels/flow/core/cmds/debug.py
+++ b/Engine/Extras/ushell/channels/flow/core/cmds/debug.py
@@ -115,20 +115,53 @@ class Wipe(flow.cmd.Cmd):
 class Paths(flow.cmd.Cmd):
     """ Prints behind-the-scenes paths. """
 
-    def main(self):
+    def gather(self):
         for item in Path(__file__).parents:
             if (item / "ushell.bat").is_file():
                 deploy_dir = item
                 break
         else:
-            deploy_path = "*unknown*"
+            deploy_path = "?unknown?"
 
         channel = self.get_channel()
         system = channel.get_system()
-        print("Deploy:", deploy_dir)
-        print(" State:", system.get_working_dir())
-        print(" Tools:", system.get_tools_dir())
-        print("  Temp:", system.get_temp_dir())
+        condition = lambda x: Path(x).resolve()
+        return {
+            "deploy"    : condition(deploy_dir),
+            "working"   : condition(system.get_working_dir()),
+            "shims"     : condition(system.get_working_dir()) / "shims",
+            "tools"     : condition(system.get_tools_dir()),
+            "temp"      : condition(system.get_temp_dir()),
+        }
+
+    def main(self):
+        for name, dir in self.gather().items():
+            print("%-8s:" % name, dir)
+
+#-------------------------------------------------------------------------------
+class Browse(Paths):
+    """ Opens internal support directories in a file browser. """
+    what = flow.cmd.Arg(str, "Supporting component's path to explore")
+    complete_what = ("deploy", "working", "tools", "temp", "shims")
+
+    def main(self):
+        paths = self.gather()
+        what = self.args.what.lower()
+        if what not in paths:
+            self.print_error(f"Not sure how to browse '{what}'")
+            return False
+
+        dest = paths[what]
+        print(what, "=", dest)
+
+        if os.name != "nt":
+            self.print_warning("...fancy a shot at an implementation for this platform?")
+            return False
+
+        explorer_arg = ("explorer.exe", "/e,/expand,", dest)
+        print("Launching:", *explorer_arg)
+        import subprocess as sp
+        sp.run(explorer_arg)
 
 #-------------------------------------------------------------------------------
 class _Channels(flow.cmd.Cmd):
diff --git a/Engine/Extras/ushell/channels/flow/core/describe.flow.py b/Engine/Extras/ushell/channels/flow/core/describe.flow.py
--- a/Engine/Extras/ushell/channels/flow/core/describe.flow.py
+++ b/Engine/Extras/ushell/channels/flow/core/describe.flow.py
@@ -49,6 +49,11 @@ debug_paths.source("cmds/debug.py", "Paths")
 debug_paths.invoke("debug", "paths")
 debug_paths.prefix("$")
 
+debug_browse = flow.describe.Command()
+debug_browse.source("cmds/debug.py", "Browse")
+debug_browse.invoke("debug", "browse")
+debug_browse.prefix("$")
+
 debug_wipe = flow.describe.Command()
 debug_wipe.source("cmds/debug.py", "Wipe")
 debug_wipe.invoke("debug", "wipe")
diff --git a/Engine/Extras/ushell/channels/flow/core/system/flow/_flick.py b/Engine/Extras/ushell/channels/flow/core/system/flow/_flick.py
--- a/Engine/Extras/ushell/channels/flow/core/system/flow/_flick.py
+++ b/Engine/Extras/ushell/channels/flow/core/system/flow/_flick.py
@@ -156,13 +156,11 @@ class Cmd(object):
                 print("", f"[{arg_name}]", end="")
             args.append((arg_name, arg.get_desc()))
 
-        suffix = ""
-        for arg_name, arg in self.read_many_args():
-            prefix = "[-- " if args else ""
-            print(f" {prefix}{arg_name}\u2026", end="")
-            suffix += "]" if args else ""
+        for arg_name, arg, with_sep in self.read_many_args():
+            prefix = "-- " if with_sep else ""
+            print(f" [{prefix}{arg_name}\u2026]", end="")
             args.append((arg_name, arg.get_desc()))
-        print(suffix)
+        print()
 
         if short:
             return
@@ -273,16 +271,10 @@ class Cmd(object):
                 break
 
             result[arg_name] = (cmd_arg.cast_value(in_arg), False)
-        else:
-            if cmd_arg:
-                if positionals:
-                    if positionals[0] != "--":
-                        raise ValueError(positionals)
-                    positionals.pop(0)
 
         # Are there any many-arg sink for what remains?
         cmd_many_arg = None
-        for arg_name, cmd_many_arg in self.read_many_args():
+        for arg_name, cmd_many_arg, _ in self.read_many_args():
             if not positionals:
                 return iter(((arg_name, cmd_many_arg),))
 
@@ -315,7 +307,16 @@ class Cmd(object):
         yield from self._read_argopts(_Arg)
 
     def read_many_args(self):
-        yield from self._read_argopts(_ManyArg)
+        it = self._read_argopts(_ManyArg)
+        staged = next(it, None)
+        if staged is None:
+            return
+        with_sep = False
+        for arg in it:
+            yield *staged, with_sep
+            staged = arg
+            with_sep = True
+        yield *staged, True
 
     def read_opts(self):
         yield from self._read_argopts(_Opt)
