Updated ushell [40537787]

#rnx

diff --git a/Engine/Extras/ushell/channels/flow/core/system/flow/_flick.py b/Engine/Extras/ushell/channels/flow/core/system/flow/_flick.py
--- a/Engine/Extras/ushell/channels/flow/core/system/flow/_flick.py
+++ b/Engine/Extras/ushell/channels/flow/core/system/flow/_flick.py
@@ -156,11 +156,13 @@ class Cmd(object):
                 print("", f"[{arg_name}]", end="")
             args.append((arg_name, arg.get_desc()))
 
-        for arg_name, arg, with_sep in self.read_many_args():
-            prefix = "-- " if with_sep else ""
-            print(f" [{prefix}{arg_name}\u2026]", end="")
+        suffix = ""
+        for arg_name, arg in self.read_many_args():
+            prefix = "[-- " if args else ""
+            print(f" {prefix}{arg_name}\u2026", end="")
+            suffix += "]" if args else ""
             args.append((arg_name, arg.get_desc()))
-        print()
+        print(suffix)
 
         if short:
             return
@@ -271,10 +273,16 @@ class Cmd(object):
                 break
 
             result[arg_name] = (cmd_arg.cast_value(in_arg), False)
+        else:
+            if cmd_arg:
+                if positionals:
+                    if positionals[0] != "--":
+                        raise ValueError(positionals)
+                    positionals.pop(0)
 
         # Are there any many-arg sink for what remains?
         cmd_many_arg = None
-        for arg_name, cmd_many_arg, _ in self.read_many_args():
+        for arg_name, cmd_many_arg in self.read_many_args():
             if not positionals:
                 return iter(((arg_name, cmd_many_arg),))
 
@@ -307,16 +315,7 @@ class Cmd(object):
         yield from self._read_argopts(_Arg)
 
     def read_many_args(self):
-        it = self._read_argopts(_ManyArg)
-        staged = next(it, None)
-        if staged is None:
-            return
-        with_sep = False
-        for arg in it:
-            yield *staged, with_sep
-            staged = arg
-            with_sep = True
-        yield *staged, True
+        yield from self._read_argopts(_ManyArg)
 
     def read_opts(self):
         yield from self._read_argopts(_Opt)
