Updated ushell [40575708]

#rnx

diff --git a/Engine/Extras/ushell/channels/flow/core/system/flow/_flick.py b/Engine/Extras/ushell/channels/flow/core/system/flow/_flick.py
--- a/Engine/Extras/ushell/channels/flow/core/system/flow/_flick.py
+++ b/Engine/Extras/ushell/channels/flow/core/system/flow/_flick.py
@@ -156,13 +156,11 @@ class Cmd(object):
                 print("", f"[{arg_name}]", end="")
             args.append((arg_name, arg.get_desc()))
 
-        suffix = ""
-        for arg_name, arg in self.read_many_args():
-            prefix = "[-- " if args else ""
-            print(f" {prefix}{arg_name}\u2026", end="")
-            suffix += "]" if args else ""
+        for arg_name, arg, with_sep in self.read_many_args():
+            prefix = "-- " if with_sep else ""
+            print(f" [{prefix}{arg_name}\u2026]", end="")
             args.append((arg_name, arg.get_desc()))
-        print(suffix)
+        print()
 
         if short:
             return
@@ -262,6 +260,7 @@ class Cmd(object):
                 positionals.append(arg)
 
         # Parse Args
+        sep_encountered = False
         cmd_arg = None
         cmd_args = self.read_args()
         for arg_name, cmd_arg in cmd_args:
@@ -269,34 +268,37 @@ class Cmd(object):
                 return iter(((arg_name, cmd_arg), *cmd_args))
 
             in_arg = positionals.pop(0)
-            if in_arg == "--":
+            if sep_encountered := (in_arg == "--"):
                 break
 
             result[arg_name] = (cmd_arg.cast_value(in_arg), False)
-        else:
-            if cmd_arg:
-                if positionals:
-                    if positionals[0] != "--":
-                        raise ValueError(positionals)
-                    positionals.pop(0)
 
-        # Are there any many-arg sink for what remains?
+            if positionals and (sep_encountered := positionals[0] == "--"):
+                positionals.pop(0)
+                break
+
+        # Are there any many-arg sinks for what remains?
         cmd_many_arg = None
-        for arg_name, cmd_many_arg in self.read_many_args():
+        for arg_name, cmd_many_arg, sep_expected in self.read_many_args():
+            if sep_expected and not sep_encountered:
+                if positionals:
+                    raise ValueError(positionals)
+                break
+
             if not positionals:
                 return iter(((arg_name, cmd_many_arg),))
 
-            onwards = False
+            sep_encountered = False
             pieces = []
             while positionals:
                 in_arg = positionals.pop(0)
-                if onwards := (in_arg == "--"):
+                if sep_encountered := (in_arg == "--"):
                     break
                 pieces.append(in_arg)
 
             result[arg_name] = (tuple(pieces), False)
 
-            if not onwards and not positionals:
+            if not sep_encountered and not positionals:
                 break
 
         if cmd_many_arg is not None:
@@ -315,7 +317,18 @@ class Cmd(object):
         yield from self._read_argopts(_Arg)
 
     def read_many_args(self):
-        yield from self._read_argopts(_ManyArg)
+        arg_num = sum(1 for x in self.read_args())
+
+        it = self._read_argopts(_ManyArg)
+        staged = next(it, None)
+        if staged is None:
+            return
+        with_sep = False
+        for arg in it:
+            yield *staged, with_sep
+            staged = arg
+            with_sep = True
+        yield *staged, bool(arg_num > 0)
 
     def read_opts(self):
         yield from self._read_argopts(_Opt)
diff --git a/Engine/Extras/ushell/channels/unreal/core/cmds/getbuild.py b/Engine/Extras/ushell/channels/unreal/core/cmds/getbuild.py
--- a/Engine/Extras/ushell/channels/unreal/core/cmds/getbuild.py
+++ b/Engine/Extras/ushell/channels/unreal/core/cmds/getbuild.py
@@ -8,6 +8,7 @@ import os
 from peafour import P4
 import urllib.request
 import string
+import fzf
 
 #-------------------------------------------------------------------------------
 class GetBuild(unrealcmd.Cmd):
@@ -20,17 +21,17 @@ By default the latest build is downloaded. You can choose a specific build via -
     * --match 1337                  - To search for the build cloest to a changelist
     * --match job=v99f3;step=39f983 - For advanced searching you can specify key values"""
 
-    buildtype  = unrealcmd.Arg(str, "The type of build to download. Usually packaged or staged builds.")
-    platform  = unrealcmd.Arg(str, "The platform of the build you want to download ")
-    runargs  = unrealcmd.Arg([str], "Additional arguments to pass to the process being launched")
+    buildtype   = unrealcmd.Arg(str, "The type of build to download. Usually packaged or staged builds.")
+    platforms   = unrealcmd.Arg([str], "The platform(s) you want to download from a build")
+    runargs     = unrealcmd.Arg([str], "Additional arguments to pass to the process being launched")
 
-    project = unrealcmd.Opt("", "The project you want to download a build for. Defaults to the current project in ushell.")
-    branch = unrealcmd.Opt("", "Branch/Stream to fetch data for, defaults to current p4 stream")
-    namespace = unrealcmd.Opt("", "The namespace to use overriding what is defined in engine ini")
-    host = unrealcmd.Opt("", "The host name of the server to connect to, defaults to your cloud storage server in engine ini")
+    project     = unrealcmd.Opt("", "The project you want to download a build for. Defaults to the current project in ushell.")
+    branch      = unrealcmd.Opt("", "Branch/Stream to fetch data for, defaults to current p4 stream")
+    namespace   = unrealcmd.Opt("", "The namespace to use overriding what is defined in engine ini")
+    host        = unrealcmd.Opt("", "The host name of the server to connect to, defaults to your cloud storage server in engine ini")
 
-    dest = unrealcmd.Opt("", "The directory to download the build to. Local contents will be overwritten. Defaults to `<projectdir>/Staged/<platform>`")
-    match = unrealcmd.Opt("", "Either a build id (hexadecimal string), a changelist (number) or a key value seperate list of key=values you want to find a build for. e.g. changelist=1337;job=foo")
+    dest        = unrealcmd.Opt("", "The directory to download the build to. Local contents will be overwritten. Defaults to `<projectdir>/Staged/<platform>`")
+    match       = unrealcmd.Opt("", "Either a build id (hexadecimal string), a changelist (number) or a key value seperate list of key=values you want to find a build for. e.g. changelist=1337;job=foo")
 
     def lookup_service_name(self):
         config = self.get_unreal_context().get_config()
@@ -124,7 +125,7 @@ By default the latest build is downloaded. You can choose a specific build via -
     def complete_buildtype(self, prefix):
         return ["packaged", "staged"]
 
-    def complete_platform(self, prefix):
+    def complete_platforms(self, prefix):
         server_platforms = ("win64", "linux", "mac")
         for platform in unrealcmd.Cmd.complete_platform(self, prefix):
             yield platform
@@ -169,20 +170,13 @@ By default the latest build is downloaded. You can choose a specific build via -
             return build_id_arg
         return None
 
-    def search_build(self, host, namespace, bucket, search_request):
+    def build_query(self, search_request):
         search_changelist = None
         if search_request.isnumeric():
             search_changelist = search_request
             print("Searching for build closest to changelist: " + str(search_changelist))
             search_request = "" # we special case handle the search for changelists
 
-        # Run a search against the api to find builds that matches the request
-        url = f"{host}/api/v2/builds/{namespace}/{bucket}/search"
-        headers = {
-            "Authorization" : f"Bearer {self._AccessToken}",
-            "Content-Type" : "application/json"
-        }
-
         query = {}
         for statement in search_request.split(';'):
             if statement.find('=') != -1:
@@ -192,18 +186,62 @@ By default the latest build is downloaded. You can choose a specific build via -
         if search_changelist:
             query["commit"] = {"$lte": search_changelist}
 
-        # we only need the first build found so limiting results to 1
-        data = json.dumps({"query": query, "options": {"limit": 1}}).encode()
+        return query
+
+    def search_build(self, host, namespace, bucket, query):
+        # Run a search against the api to find builds that matches the request
+        url = f"{host}/api/v2/builds/{namespace}/{bucket}/search"
+        headers = {
+            "Authorization" : f"Bearer {self._AccessToken}",
+            "Content-Type" : "application/json"
+        }
+
+        data = json.dumps({"query": query}).encode()
         request = urllib.request.Request(url, data=data, headers=headers)
         response = urllib.request.urlopen(request)
         body = response.read().decode()
 
         results = json.loads(body)["results"]
-        if len(results) == 0:
-            return (None, None, query)
-
-        return (results[0]["buildId"], results[0]["metadata"], query)
-
+        return [{'name' : r['metadata']['name'], 'platforms': {r['metadata']['platform']: r['buildId']}, 'commit': r['metadata'].get('commit', '')} for r in results]
+
+    def _show_select_build_prompt(self, branch_cl, available_builds):
+        closest_cl = 1 << 30
+        p4_change_found = False
+        active_platforms = set()
+        for build in available_builds:
+            cl = build['commit']                 
+            active_platforms.update(build['platforms'].keys())
+
+            if type(cl) is int or (type(cl) is str and cl.isnumeric()):
+                p4_change_found = True
+                # is a changelist number
+                if abs(int(cl) - branch_cl) < abs(closest_cl - branch_cl):
+                    closest_cl = int(cl)
+
+        def abbreviate(x):
+            vowels_plus = ("a", "e", "i", "o", "u", "y")
+            x = "".join(x for x in x if x not in vowels_plus)
+            return x[:7]
+        columns = {x:abbreviate(x) for x in active_platforms}
+
+        def list_available(closest_cl, available_builds):
+            for build in available_builds:
+                cl = build['commit']
+                name = build['name']
+                platforms = build['platforms'].keys()
+                middle = ""
+                for column_platform,short_name in columns.items():
+                    middle += " "
+                    middle += short_name if column_platform in platforms else (" " * len(short_name))
+                if cl != None:
+                    if not p4_change_found:
+                        closest_cl = branch_cl
+                    cl = str(cl) + ("*" if cl == closest_cl else " ")
+                yield "%10s%s %s" % (cl, middle, name)
+
+        for line in fzf.run(list_available(closest_cl, available_builds), sort=False):
+            return line.split()[-1]
+        
     def main(self):
         ue_context = self.get_unreal_context()
         engine = ue_context.get_engine()
@@ -213,15 +251,21 @@ By default the latest build is downloaded. You can choose a specific build via -
             project = ue_context.get_project().get_name()
             project_path = os.path.dirname(ue_context.get_project().get_path())
 
-        platform = self.convert_platform(self.args.platform)
+        if len(self.args.platforms) == 0:
+            self.print_error("You need to select at least one platform")
+            return False
+        
         build_type = self.convert_buildtype(self.args.buildtype)
         dest = self.args.dest
         if dest == "" and project_path:
-            dest = os.path.join(project_path, self.args.buildtype, platform)
+            dest = os.path.join(project_path, self.args.buildtype.capitalize())
 
         build_id = self.get_build_id(self.args.match)
         branch = self._get_p4_stream()
 
+        branch_cl = ue_context.get_engine().get_info()["Changelist"]
+        print("Branch synced to:", branch_cl)
+    
         if self.args.namespace:
             namespace = self.args.namespace
         else:
@@ -235,34 +279,71 @@ By default the latest build is downloaded. You can choose a specific build via -
         if not self.refresh_access_token():
             return False
 
-        bucket = f"{project}.{build_type}.{branch}.{platform}".lower()
-
         if not host:
             return False
         if not namespace:
             return False
 
-        if not build_id:
-            self.print_info("Enumerating builds")
-            # Search for build id
-            (build_id, metadata, query) = self.search_build(host, namespace, bucket, self.args.match)
-            if not build_id:
-                return f"Failed to find any build that matches query {query} in namespace {namespace} and bucket {bucket}"
+        platforms_to_download = [self.convert_platform(p) for p in self.args.platforms]
+
+        builds_to_download = []
+        if build_id:
+            # a single build id was specified, this can only be used wit ha single platform as well
+            if len(platforms_to_download) != 1:
+                self.print_error("Only a single platform can be specified when specifying a build id")
+                return False
+
+            builds_to_download = [(bucket, build_id, platforms_to_download[0])]
+        else:
+            # we need to search
+            found_builds = {}
+
+            query = self.build_query(self.args.match)
+            # Scan all platforms to see which ones have a build that matches
+            for platform in platforms_to_download:
+                self.print_info(f"Enumerating builds for platform {platform}")
+                bucket = f"{project}.{build_type}.{branch}.{platform}".lower()
+
+                builds = self.search_build(host, namespace, bucket, query)
+                print(f"Found {len(builds)} builds")
+                for build in builds:
+                    if build['name'] in found_builds:
+                        # this build already existed but a new platform was found add it to the list
+                        found_builds[build['name']]['platforms'].update(build['platforms'])
+                    else:
+                        found_builds[build['name']] = build
+
+            # remove builds that only have some of the platforms we want
+            found_builds = {k:v for k, v in found_builds.items() if len(v["platforms"]) == len(platforms_to_download)}
+
+            if found_builds == None or len(found_builds) == 0:
+                self.print_error(f"Failed to find any build that matches query {query} in namespace {namespace} and bucket {bucket}")
+                return False
+
+            selected_build = None
+            if len(found_builds) == 1:
+                print(f"Single build found, using build {found_builds[0]['name']}")
+
+                selected_build = found_builds[0]
             else:
-                build_name = metadata['name']
-                print(f"Found build {build_id} - {build_name}")
-        access_token_env_var = "CloudAccessToken"
+                # Ask the user to select a build
+                print("Select build to fetch:", end="")
+                selected_build_name = self._show_select_build_prompt(branch_cl, found_builds.values())
+                if not selected_build_name:
+                    self.print_error("No build was selected")
+                    return False
+                
+                print(f"Using build {selected_build_name}")
+                selected_build = found_builds[selected_build_name]
+
+            # build the selection of build ids to download from which buckets
+            for platform, build_id in selected_build['platforms'].items():
+                if not platform in platforms_to_download:
+                    continue
+                bucket = f"{project}.{build_type}.{branch}.{platform}".lower()
+                builds_to_download.extend([(bucket, build_id, platform)])
 
-        args = (
-            "builds",
-            "download",
-            "--url=" + str(host),
-            "--namespace=" + str(namespace),
-            "--bucket=" + bucket,
-            "--access-token-env=" + access_token_env_var,
-            dest,
-            build_id
-        )
+        access_token_env_var = "CloudAccessToken"
 
         bin_type = "Win64"
         if sys.platform == "darwin": bin_type = "Mac"
@@ -276,7 +357,27 @@ By default the latest build is downloaded. You can choose a specific build via -
         bin_path /= f"Binaries/{bin_type}/zen{bin_ext}"
         if not bin_path.is_file():
             raise FileNotFoundError(f"Unable to find '{bin_path}'")
-        args = (*args, *self.args.runargs)
-        cmd = self.get_exec_context().create_runnable(bin_path, *args)
-        result = cmd.run()
+        
+        result = 0
+        for bucket, build_id, build_platform in builds_to_download:
+            if result:
+                self.print_error(f"Returned error exit code {result}, aborting download")
+                return False
+            
+            build_dest = os.path.join(dest, build_platform)
+            args = (
+                "builds",
+                "download",
+                "--url=" + str(host),
+                "--namespace=" + str(namespace),
+                "--bucket=" + bucket,
+                "--access-token-env=" + access_token_env_var,
+                build_dest,
+                build_id
+            )
+
+            args = (*args, *self.args.runargs)
+            cmd = self.get_exec_context().create_runnable(bin_path, *args)
+            result = cmd.run()
+
         return result
diff --git a/Engine/Extras/ushell/channels/unreal/perforce/cmds/clean.py b/Engine/Extras/ushell/channels/unreal/perforce/cmds/clean.py
--- a/Engine/Extras/ushell/channels/unreal/perforce/cmds/clean.py
+++ b/Engine/Extras/ushell/channels/unreal/perforce/cmds/clean.py
@@ -7,6 +7,56 @@ from peafour import P4
 import subprocess as sp
 from pathlib import Path
 
+#-------------------------------------------------------------------------------
+def _running_exes_powershell():
+    ps_cmd = (
+        "powershell.exe",
+        "-NoProfile",
+        "-NoLogo",
+        "-Command",
+        r'(Get-CimInstance -ClassName Win32_Process).ExecutablePath',
+    )
+    startupinfo = sp.STARTUPINFO()
+    startupinfo.dwFlags = sp.STARTF_USESHOWWINDOW
+    startupinfo.wShowWindow = sp.SW_HIDE
+    proc = sp.Popen(
+        ps_cmd,
+        stdout=sp.PIPE, stderr=sp.DEVNULL, stdin=sp.DEVNULL,
+        startupinfo=startupinfo, creationflags=sp.CREATE_NEW_CONSOLE
+    )
+    with proc:
+        yield from (x.strip().decode() for x in proc.stdout)
+        proc.stdout.close()
+
+#-------------------------------------------------------------------------------
+def _running_exes_wmic():
+    proc = sp.Popen(
+        ("wmic.exe", "process", "get", "executablepath"),
+        stdout=sp.PIPE, stderr=sp.DEVNULL, stdin=sp.DEVNULL
+    )
+    with proc:
+        for line in proc.stdout:
+            line = line.strip().decode()
+            if line:
+                yield line
+        proc.stdout.close()
+
+#-------------------------------------------------------------------------------
+def _read_running_exes():
+    try:
+        gen = _running_exes_powershell()
+    except:
+        gen = _running_exes_wmic()
+    yield from gen
+
+#-------------------------------------------------------------------------------
+def _detect_locked_files(root_dir):
+    for line in _read_running_exes():
+        if line and Path(line).parent.is_relative_to(root_dir):
+            return Path(line)
+
+
+
 #-------------------------------------------------------------------------------
 class _Collector(object):
     def __init__(self):
@@ -159,23 +209,6 @@ class Clean(flow.cmd.Cmd):
             if sub_dir.name.lower() not in self._saved_keeps:
                 collector.add_dir(sub_dir)
 
-    def _detect_locked_files(self, root_dir):
-        proc = sp.Popen(
-            ("wmic.exe", "process", "get", "executablepath"),
-            stdout=sp.PIPE, stderr=sp.DEVNULL
-        )
-
-        ret = False
-        for line in proc.stdout.readlines():
-            line = line.strip().decode()
-            if line and Path(line).parent.is_relative_to(root_dir):
-                ret = line
-                break
-
-        proc.stdout.close()
-        proc.wait()
-        return ret
-
     def main(self):
         self._saved_keeps = {x.strip().lower() for x in self.args.savedkeeps.split(",")}
 
@@ -185,9 +218,18 @@ class Clean(flow.cmd.Cmd):
 
         if os.name == "nt":
             self.print_info("Checking running processes")
-            if running_exe := self._detect_locked_files(branch.get_dir()):
+            for _ in range(2):
+                running_exe = _detect_locked_files(branch.get_dir())
+                if not running_exe:
+                    break
+
+                if running_exe.name == "dotnet.exe":
+                    print("'dotnet.exe' is playing gooseberry")
+                    print("...lets terminate them")
+                    sp.run(("taskkill", "/f", "/im", str(running_exe.name)))
+                    continue
+
                 raise RuntimeError(f"Not cleaning because '{running_exe}' is running")
-                return False
 
         self.print_info("Finding directories and files to clean")
         root_dirs = [
