Updated ushell [40613888]

#rnx

diff --git a/Engine/Extras/ushell/channels/unreal/perforce/pylib/p4utils/__init__.py b/Engine/Extras/ushell/channels/unreal/perforce/pylib/p4utils/__init__.py
--- a/Engine/Extras/ushell/channels/unreal/perforce/pylib/p4utils/__init__.py
+++ b/Engine/Extras/ushell/channels/unreal/perforce/pylib/p4utils/__init__.py
@@ -58,7 +58,7 @@ def create_p4config(p4config_path, client, username, port=None):
         print("P4CLIENT=", client, **print_args)
         print("P4USER=", username, **print_args)
         if port:
-            print("P4PORT", port, **print_args)
+            print("P4PORT=", port, **print_args)
 
 #-------------------------------------------------------------------------------
 def ensure_p4config(start_dir=None):
@@ -75,9 +75,12 @@ def ensure_p4config(start_dir=None):
         return
     client, root_dir = client
 
+    # Forward whatever P4PORT is set
+    port = get_p4_set("P4PORT")
+
     # Now we know where to locate a p4config file
     p4config_path = f"{root_dir}/{p4config_name}"
-    create_p4config(p4config_path, client, username)
+    create_p4config(p4config_path, client, username, port)
     return p4config_path, True
 
 #-------------------------------------------------------------------------------
