Updated ushell [40647405]

#rnx

diff --git a/Engine/Extras/ushell/.version b/Engine/Extras/ushell/.version
new file mode 100644
--- /dev/null
+++ b/Engine/Extras/ushell/.version
@@ -0,0 +1 @@
+40647405
\ No newline at end of file
diff --git a/Engine/Extras/ushell/channels/unreal/core/cmds/getbuild.py b/Engine/Extras/ushell/channels/unreal/core/cmds/getbuild.py
--- a/Engine/Extras/ushell/channels/unreal/core/cmds/getbuild.py
+++ b/Engine/Extras/ushell/channels/unreal/core/cmds/getbuild.py
@@ -293,6 +293,8 @@ By default the latest build is downloaded. You can choose a specific build via -
                 self.print_error("Only a single platform can be specified when specifying a build id")
                 return False
 
+            platform = platforms_to_download[0]
+            bucket = f"{project}.{build_type}.{branch}.{platform}".lower()
             builds_to_download = [(bucket, build_id, platforms_to_download[0])]
         else:
             # we need to search
@@ -322,9 +324,8 @@ By default the latest build is downloaded. You can choose a specific build via -
 
             selected_build = None
             if len(found_builds) == 1:
-                print(f"Single build found, using build {found_builds[0]['name']}")
-
-                selected_build = found_builds[0]
+                selected_build = next(iter(found_builds.values()))
+                print(f"Single build found, using build {selected_build["name"]}")
             else:
                 # Ask the user to select a build
                 print("Select build to fetch:", end="")
