Updated ushell [40709780]
#rnx

diff --git a/Engine/Extras/ushell/.version b/Engine/Extras/ushell/.version
--- a/Engine/Extras/ushell/.version
+++ b/Engine/Extras/ushell/.version
@@ -1 +1 @@
-40647405
\ No newline at end of file
+40709780
\ No newline at end of file
diff --git a/Engine/Extras/ushell/channels/unreal/core/cmds/getbuild.py b/Engine/Extras/ushell/channels/unreal/core/cmds/getbuild.py
--- a/Engine/Extras/ushell/channels/unreal/core/cmds/getbuild.py
+++ b/Engine/Extras/ushell/channels/unreal/core/cmds/getbuild.py
@@ -139,7 +139,11 @@ By default the latest build is downloaded. You can choose a specific build via -
                 yield item.name + "/"
 
     def convert_platform(self, platform):
-        return self.get_platform(platform).get_config_name()
+        try:
+            return self.get_platform(platform).get_config_name()
+        except ValueError:
+            # unknown platform, just assume the name it valid
+            return platform 
 
     def convert_buildtype(self, build_type):
         if build_type == "packaged":
@@ -277,11 +281,14 @@ By default the latest build is downloaded. You can choose a specific build via -
             host = self.lookup_host()
 
         if not self.refresh_access_token():
+            self.print_error("Failed to acquire a access token")
             return False
 
         if not host:
+            self.print_error("No host specified")
             return False
         if not namespace:
+            self.print_error("Namespace was not resolved, unable to determine what to download.")
             return False
 
         platforms_to_download = [self.convert_platform(p) for p in self.args.platforms]
@@ -325,7 +332,8 @@ By default the latest build is downloaded. You can choose a specific build via -
             selected_build = None
             if len(found_builds) == 1:
                 selected_build = next(iter(found_builds.values()))
-                print(f"Single build found, using build {selected_build["name"]}")
+                selected_build_name = selected_build["name"]
+                print(f"Single build found, using build {selected_build_name}")
             else:
                 # Ask the user to select a build
                 print("Select build to fetch:", end="")
@@ -339,8 +347,6 @@ By default the latest build is downloaded. You can choose a specific build via -
 
             # build the selection of build ids to download from which buckets
             for platform, build_id in selected_build['platforms'].items():
-                if not platform in platforms_to_download:
-                    continue
                 bucket = f"{project}.{build_type}.{branch}.{platform}".lower()
                 builds_to_download.extend([(bucket, build_id, platform)])
 
diff --git a/Engine/Extras/ushell/channels/unreal/core/cmds/insights.py b/Engine/Extras/ushell/channels/unreal/core/cmds/insights.py
new file mode 100644
--- /dev/null
+++ b/Engine/Extras/ushell/channels/unreal/core/cmds/insights.py
@@ -0,0 +1,63 @@
+import unreal
+import unrealcmd
+import subprocess as sp
+
+#-------------------------------------------------------------------------------
+class Insights(unrealcmd.Cmd):
+    """ Conveninent access to UnrealInsights """
+    trace       = unrealcmd.Arg("", "Trace ident to open from server or .utrace file")
+    uiargs      = unrealcmd.Arg([str], "Additional arguments to pass to UnrealInsights")
+    build       = unrealcmd.Opt(False, "Build UnrealInsights before launching")
+    attach      = unrealcmd.Opt(False, "Attach the debugger")
+    debug       = unrealcmd.Opt(False, "Use debug variant binaries")
+    noautobuild = unrealcmd.Opt(False, "Skip check for binaries and inference of --build")
+
+    def _get_lastest_trace(self):
+        # todo: we shall create a UTS client!
+        import os
+        if os.name != "nt":
+            self.print_warning("The status of 'latest'; not yet")
+            return
+
+        from pathlib import Path
+        zerozeroone = Path(os.getenv("localappdata"))
+        zerozeroone /= "UnrealEngine/Common/UnrealTrace/Store/001"
+
+        items = list(zerozeroone.glob("*.utrace"))
+        if not items:
+            return
+        items.sort(key=lambda x: x.stat().st_mtime)
+        return items[-1]
+
+    def _has_binary(self, *, debug=False):
+        ue_context = self.get_unreal_context()
+        if not (target := ue_context.get_target_by_name("UnrealInsights")):
+            return False
+        if not (build := target.get_build()):
+            return False
+        return build.get_binary_path() is not None
+
+    def main(self):
+        # Infer --build if there doesn't appear to be binaries.
+        if not self.args.noautobuild:
+            if not self._has_binary(debug=self.args.debug):
+                self.args.build = True
+
+        # Inject the file to the latest trace if requested
+        pass_through = self.args.uiargs
+        if self.args.trace == "latest":
+            if latest := self._get_lastest_trace():
+                pass_through = (latest, *pass_through)
+
+        # Yup, this is mostly a shim around other ushell commands!
+        ui_args = (
+            "_run",
+            "program",
+            "UnrealInsights",
+            "debug" if self.args.debug else "development",
+        )
+        if self.args.build:     ui_args = (*ui_args, "--build")
+        if self.args.attach:    ui_args = (*ui_args, "--attach")
+        if pass_through:        ui_args = (*ui_args, "--", *pass_through)
+
+        sp.run(ui_args)
diff --git a/Engine/Extras/ushell/channels/unreal/core/cmds/perftest.py b/Engine/Extras/ushell/channels/unreal/core/cmds/perftest.py
new file mode 100644
--- /dev/null
+++ b/Engine/Extras/ushell/channels/unreal/core/cmds/perftest.py
@@ -0,0 +1,370 @@
+# Copyright Epic Games, Inc. All Rights Reserved.
+
+from enum import Enum
+from pathlib import Path
+import unrealcmd
+import unreal
+
+#-------------------------------------------------------------------------------
+class PerfTestType(Enum):
+    SEQUENCE    = 0x00
+    REPLAY      = 0x01
+    MATERIAL    = 0x02
+    CAMERA      = 0x03
+
+    def as_arg(self):
+        match self:
+            case PerfTestType.SEQUENCE:
+                return "AutomatedPerfTest.SequenceTest"
+            case PerfTestType.REPLAY:
+                return "AutomatedPerfTest.ReplayTest"
+            case PerfTestType.MATERIAL:
+                return "AutomatedPerfTest.MaterialTest"
+            case PerfTestType.CAMERA:
+                return "AutomatedPerfTest.StaticCameraTest"
+
+
+#-------------------------------------------------------------------------------
+class PerfSubTestType(Enum):
+    PERF     = 0x00
+    LLM      = 0x01
+    INSIGHTS = 0x02
+    GPUPERF  = 0x03
+    ALL      = 0x04
+
+    def as_arg(self):
+        match self:
+            case PerfSubTestType.PERF:
+                return "AutomatedPerfTest.DoPerf"
+            case PerfSubTestType.LLM:
+                return "AutomatedPerfTest.DoLLM"
+            case PerfSubTestType.INSIGHTS:
+                return "AutomatedPerfTest.DoInsightsTrace"
+            case PerfSubTestType.GPUPERF:
+                return "AutomatedPerfTest.DoGPUPerf"
+
+    @staticmethod
+    def parse(external:str) -> "PerfSubTestType":
+        try: return PerfSubTestType[external.upper()]
+        except: raise ValueError(f"Unknown Sub Test type '{external}'")
+
+
+#-------------------------------------------------------------------------------
+class PerfTestResult(Enum):
+    SUCCESS     = 0
+    ERROR       = 1
+
+    def __get__(self, instance, owner):
+        return self.value
+
+
+#-------------------------------------------------------------------------------
+class _PerfTestConfig(object):
+    iterations  = {}
+    resx        = 0
+    resy        = 0
+    ignore_log  = False
+
+    def load_defaults(self, xml_text : str):
+        import xml.etree.ElementTree as ET
+        tree = ET.fromstring(xml_text)
+
+        DEFAULT_VAL_KEY = 'DefaultValue'
+        def _get_option(name:str, attrib_key = DEFAULT_VAL_KEY) -> str:
+            # Currently only checks for options. Properties may need
+            # variable evaluation and I don't think we want to go there
+            # yet.
+            try:
+                return tree.find(".//{*}Option" + f"[@Name='{name}']").attrib[attrib_key]
+            except:
+                return None
+
+        def _get_bool(name:str, attrib_key = DEFAULT_VAL_KEY) -> bool:
+            val = _get_option(name, attrib_key)
+            return True if val is not None and val.lower() == "true" else False
+
+        # More default values can be extracted if needed.
+        self.iterations = {
+            PerfSubTestType.PERF     : _get_option("IterationsPerf"),
+            PerfSubTestType.LLM      : _get_option("IterationsLLM"),
+            PerfSubTestType.INSIGHTS : _get_option("IterationsInsights"),
+            PerfSubTestType.GPUPERF  : _get_option("IterationsGPUPerf")
+        }
+
+        self.resx       = _get_option("ResX")
+        self.resy       = _get_option("ResY")
+        self.ignore_log = _get_bool("APTIgnoreTestBuildLogging")
+
+    def get_iterations(self, subtest : PerfSubTestType) -> int:
+        if subtest in self.iterations:
+            if self.iterations[subtest] is not None:
+                return self.iterations[subtest]
+
+        # Return default values if a subtest iteration value
+        # does not exist.
+        DEFAULT_ITERATIONS  = 3
+        PERF_ITERATIONS     = 6
+        LLM_ITERATIONS      = 3
+        INSIGHTS_ITERATIONS = 1
+        GPUPERF_ITERATIONS  = 3
+
+        match subtest:
+            case PerfSubTestType.PERF:
+                return PERF_ITERATIONS
+            case PerfSubTestType.LLM:
+                return LLM_ITERATIONS
+            case PerfSubTestType.INSIGHTS:
+                return INSIGHTS_ITERATIONS
+            case PerfSubTestType.GPUPERF:
+                return GPUPERF_ITERATIONS
+        return DEFAULT_ITERATIONS
+
+
+#-------------------------------------------------------------------------------
+class PerfTestBaseCmd(unrealcmd.Cmd):
+    """
+    Automated Perf Test(APT) - This command runs the specified performance test on
+    a given APT plugin-enabled project and generates a performance report. You can choose
+    to run any or all of the following sub tests:
+
+    - Perf:     This is the default perf run. This collects timing metrics such as
+    GameThread time, RHIThread time, Average FPS etc.
+
+    - LLM:      This subtest collects memory related metrics. The project needs LLM
+    tracking enabled for this to work as expected.
+
+    - GPUPerf:  This subtest collects GPU/Graphics related performance metrics and is
+    suitable for testing shader performance etc.
+
+    - Insights: Enables insights trace and runs the test. A utrace file is generated
+    which can be used to further do performance investigations.
+    """
+    host        = unreal.Platform.get_host()
+    platform    = unrealcmd.Arg(host,       "Platform to run the test on")
+    subtest     = unrealcmd.Arg("all",      "Subtest Type")
+    variant     = unrealcmd.Arg("test",     "Build variant to use for the test")
+    target      = unrealcmd.Arg("game",     "Cooked target to use")
+    repeat      = unrealcmd.Opt(0,          "Number of iterations to run the test")
+    build       = unrealcmd.Opt("",         "Use an alternative staged directory")
+    targetname  = unrealcmd.Opt("",         "Override target name")
+    uatargs     = unrealcmd.Arg([str],      "Additional arguments to pass to UAT")
+    # --- Miscellaneous
+    resx        = unrealcmd.Opt(0,          "Res X")
+    resy        = unrealcmd.Opt(0,          "Res Y")
+    fps_chart   = unrealcmd.Opt(False,      "Enable FPS Chart")
+    debug_mem   = unrealcmd.Opt(False,      "Enable debug memory on platform if available")
+    testid      = unrealcmd.Opt("",         "Test ID for perf test")
+    # --- Config
+    config      = _PerfTestConfig()
+
+    def run_perf_subtest(self, test_type: PerfTestType, subtest: PerfSubTestType, *args):
+        self.print_info(f"Initializing {test_type.name.lower()} Test")
+        print("Arguments: ", *args)
+
+        ue_context      = self.get_unreal_context()
+        platform_name   = self.args.platform.lower()
+        target_type     = unreal.TargetType.parse(self.args.target)
+        target          = ue_context.get_target_by_type(target_type)
+        iterations      = self.args.repeat
+        project         = ue_context.get_project()
+        platform        = self.get_platform(platform_name)
+        cook_form       = platform.get_cook_form(target_type.name.lower())
+        stage_dir       = project.get_dir() / f"Saved/StagedBuilds/{cook_form}"
+
+        # Load defaults from config if required.
+        iterations  = iterations if iterations != 0 else self.config.get_iterations(subtest)
+        resx        = self.args.resx if self.args.resx != 0 else self.config.resx
+        resy        = self.args.resy if self.args.resy != 0 else self.config.resy
+        debug_mem   = self.args.debug_mem or subtest == PerfSubTestType.LLM or subtest == PerfSubTestType.INSIGHTS
+        fps_chart   = self.args.fps_chart
+
+        if not project:
+            self.print_error("No project found in context")
+            return PerfTestResult.ERROR
+
+        test_id = f"{project.get_name()}-{self.args.subtest}-autoperftest-ushell"
+        test_id = test_id if self.args.testid == "" else self.args.testid
+
+        if not target:
+            self.print_error(f"'{project.get_name()}' does not appear to have a '{self.args.target}' target")
+            return PerfTestResult.ERROR
+
+        target_name = target.get_name()
+        if self.args.targetname != "":
+            target_name = self.args.targetname
+
+        if self.args.build != "":
+            build_dir = Path(self.args.build)
+            if build_dir.exists() and (build_dir/ "Engine").exists():
+                stage_dir = build_dir
+                print(f"Using staged build path in args")
+            else:
+                self.print_error(f"Directory {build_dir} does not exist or expected 'Engine' folder not present")
+                return PerfTestResult.ERROR
+
+        self.print_info("Running Automated Perf Test")
+        print("Sub Test:",      subtest.name.lower())
+        print("Iterations:",    iterations)
+        print("Variant:",       self.args.variant)
+        print("Target:",        target_name)
+        print("Platform:",      self.args.platform)
+        print("Staged Build:",  stage_dir)
+        print("------")
+
+        perf_args = (
+            f"-test={test_type.as_arg()}",
+            f"-{subtest.as_arg()}",
+            f"-platform={platform_name}",
+            f"-configuration={self.args.variant}",
+            f"-iterations={iterations}",
+            f"-target={target_name}",
+            f"-build={stage_dir}",
+            f"-AutomatedPerfTest.TestID={test_id}",
+            "-AutomatedPerfTest.IgnoreTestBuildLogging" if self.config.ignore_log else None,
+            "-AutomatedPerfTest.UsePlatformExtraDebugMemory" if debug_mem else None,
+            "-AutomatedPerfTest.DoCSVProfiler",
+            "-AutomatedPerfTest.TraceChannels=default,screenshot,stats",
+            "-AutomatedPerfTest.DoFPSChart" if fps_chart else None,
+            f"-resX={resx}",
+            f"-resY={resy}",
+            "-LocalReports",
+            *args
+        )
+
+        return self.run_perftest_uat_command(*perf_args)
+
+    def run_perftest(self, test_type: PerfTestType, *args):
+        self._load_defaults()
+        subtest = PerfSubTestType.parse(self.args.subtest)
+        if subtest == PerfSubTestType.ALL:
+            results:list[int] = []
+            subtests = [test for test in PerfSubTestType if test != PerfSubTestType.ALL]
+            self.print_info(f"Running all available subtests")
+            print([test.name.lower() for test in subtests])
+            for perf_subtest in subtests:
+                result = self.run_perf_subtest(test_type, perf_subtest, *args)
+                results.append(result)
+            return PerfTestResult.SUCCESS if all(ret == PerfTestResult.SUCCESS for ret in results) else PerfTestResult.ERROR
+        else:
+            return self.run_perf_subtest(test_type, subtest, *args)
+
+    def run_perftest_uat_command(self, *args):
+        dot_uat_cmd = "_uat"
+        dot_uat_args = (
+            "RunUnreal",
+            "--",
+            *args
+        )
+        exec_context = self.get_exec_context()
+        cmd = exec_context.create_runnable(dot_uat_cmd, *dot_uat_args)
+        return cmd.run()
+
+    def _load_defaults(self):
+        SETTINGS_XML = "AutomatedPerfTestCommonSettings.xml"
+        ue_context   = self.get_unreal_context()
+        engine_dir   = ue_context.get_engine().get_dir()
+        engine_dir   = engine_dir.joinpath("Plugins", "Performance", "AutomatedPerfTesting")
+        settings     = self._find_file(SETTINGS_XML, "xml", engine_dir)
+
+        if not settings:
+            self.print_warning(f"Could not find Perf Test Common Settings file.")
+            return
+
+        xml_text = settings.read_text()
+        self.config.load_defaults(xml_text)
+
+        self.print_info("Loaded config defaults")
+
+    def _find_file(self, name, ext = "*", dir:Path=None) -> Path:
+        name = name.lower()
+        for path in self._find_files(f"**/*.{ext}", dir):
+            if path.name.lower() == name:
+                return path
+        return None
+
+    def _find_files(self, pattern, dir:Path=None):
+        if dir is None:
+            # Default to Project Directory
+            ue_context = self.get_unreal_context()
+            project = ue_context.get_project()
+            if not project:
+                return None
+            dir = project.get_dir()
+        return dir.glob(pattern)
+
+    #--- Tab Auto Complete
+
+    def complete_subtest(self, prefix):
+        yield from (subtest.name.lower() for subtest in PerfSubTestType)
+
+    def complete_target(self, prefix):
+        yield from (target.name.lower() for target in unreal.TargetType)
+
+
+#-------------------------------------------------------------------------------
+class Sequence(PerfTestBaseCmd):
+    """
+    Sequence Test - Runs a configured sequence in the project while collecting
+    performance metrics.
+    """
+    SequenceComboName   = unrealcmd.Arg(str,  "Sequence Combo Name")
+    config_section      = "/Script/AutomatedPerfTesting.AutomatedSequencePerfTestProjectSettings"
+    config_key          = "MapsAndSequencesToTest"
+
+    def main(self):
+        args = (f"-AutomatedPerfTest.SequencePerfTest.MapSequenceName={self.args.SequenceComboName}",)
+        return self.run_perftest(PerfTestType.SEQUENCE, *args)
+
+    def _get_available_sequences(self) -> list[str]:
+        config = self.get_unreal_context().get_config()
+        sequences = config.get("Engine", self.config_section, self.config_key)
+        return [sequence.ComboName for sequence in sequences]
+
+    #--- Tab Auto Complete
+
+    def complete_SequenceComboName(self, prefix):
+        yield from self._get_available_sequences()
+
+
+#-------------------------------------------------------------------------------
+class Replay(PerfTestBaseCmd):
+    """
+    Replay Test - Runs performance test with given replay file.
+    """
+    ReplayFile  = unrealcmd.Arg(str,  "Replay file")
+
+    def _find_replay(self, name):
+        path = Path(name)
+        if path.is_file():
+            return path
+        self.print_info("Searching for replay file...")
+        return self._find_file(name, "replay")
+
+    def main(self):
+        replay = self._find_replay(self.args.ReplayFile)
+        if replay is None:
+            self.print_error(f"Could not find '{self.args.ReplayFile}'")
+            return 1
+        args = (f"-AutomatedPerfTest.ReplayPerfTest.ReplayName={replay}",)
+        return self.run_perftest(PerfTestType.REPLAY, *args)
+
+
+#-------------------------------------------------------------------------------
+class Material(PerfTestBaseCmd):
+    """
+    Material Test.
+    """
+    def main(self):
+        args = (None,)
+        return self.run_perftest(PerfTestType.MATERIAL, *args)
+
+
+#-------------------------------------------------------------------------------
+class StaticCamera(PerfTestBaseCmd):
+    """
+    Static Camera Test.
+    """
+    MapName  = unrealcmd.Arg(str,  "Static Camera Test Map Name")
+    def main(self):
+        args = (f"-AutomatedPerfTest.StaticCameraPerfTest.MapName={self.args.MapName}",)
+        return self.run_perftest(PerfTestType.CAMERA, *args)
\ No newline at end of file
diff --git a/Engine/Extras/ushell/channels/unreal/core/cmds/run_include_tool.py b/Engine/Extras/ushell/channels/unreal/core/cmds/run_include_tool.py
deleted file mode 100644
diff --git a/Engine/Extras/ushell/channels/unreal/core/describe.flow.py b/Engine/Extras/ushell/channels/unreal/core/describe.flow.py
--- a/Engine/Extras/ushell/channels/unreal/core/describe.flow.py
+++ b/Engine/Extras/ushell/channels/unreal/core/describe.flow.py
@@ -112,11 +112,6 @@ run_target = flow.describe.Command()
 run_target.source("cmds/run.py", "Target")
 run_target.invoke("run", "target")
 
-#-------------------------------------------------------------------------------
-run_include_tool = flow.describe.Command()
-run_include_tool.source("cmds/run_include_tool.py", "RunIncludeTool")
-run_include_tool.invoke("run", "includetool")
-
 #-------------------------------------------------------------------------------
 kill = flow.describe.Command()
 kill.source("cmds/kill.py", "Kill")
@@ -257,6 +252,30 @@ gather = flow.describe.Command()
 gather.source("cmds/gather.py", "Gather")
 gather.invoke("ushell", "gather")
 
+#-------------------------------------------------------------------------------
+perf_test_sequence = flow.describe.Command()
+perf_test_sequence.source("cmds/perftest.py", "Sequence")
+perf_test_sequence.invoke("perf", "test", "sequence")
+
+perf_test_replay = flow.describe.Command()
+perf_test_replay.source("cmds/perftest.py", "Replay")
+perf_test_replay.invoke("perf", "test", "replay")
+
+perf_test_material = flow.describe.Command()
+perf_test_material.source("cmds/perftest.py", "Material")
+perf_test_material.invoke("perf", "test", "material")
+
+perf_test_camera = flow.describe.Command()
+perf_test_camera.source("cmds/perftest.py", "StaticCamera")
+perf_test_camera.invoke("perf", "test", "camera")
+
+#-------------------------------------------------------------------------------
+perf_insights = flow.describe.Command()
+perf_insights.source("cmds/insights.py", "Insights")
+perf_insights.invoke("perf", "insights")
+
+
+
 #-------------------------------------------------------------------------------
 prompt = flow.describe.Command()
 prompt.source("prompt.py", "Prompt")
diff --git a/Engine/Extras/ushell/channels/unreal/perforce/cmds/workspace.py b/Engine/Extras/ushell/channels/unreal/perforce/cmds/workspace.py
--- a/Engine/Extras/ushell/channels/unreal/perforce/cmds/workspace.py
+++ b/Engine/Extras/ushell/channels/unreal/perforce/cmds/workspace.py
@@ -171,7 +171,9 @@ class Workspace(flow.cmd.Cmd):
             p4config_path = client_dir / p4utils.get_p4config_name()
             with p4config_path.open("wt") as out:
                 print("P4CLIENT=", client_name, sep="", file=out)
-                print("P4USER=", username)
+                print("P4USER=", username, sep="", file=out)
+                if port := p4utils.get_p4_set("P4PORT"):
+                    print("P4PORT=", port, sep="", file=out)
             print("P4CONFIG:", p4config_path)
 
         # Is this an Unreal branch?
