Updated ushell [40827618]
#rnx

diff --git a/Engine/Extras/ushell/.version b/Engine/Extras/ushell/.version
--- a/Engine/Extras/ushell/.version
+++ b/Engine/Extras/ushell/.version
@@ -1 +1 @@
-40709780
\ No newline at end of file
+40827618
\ No newline at end of file
diff --git a/Engine/Extras/ushell/channels/unreal/core/cmds/uat.py b/Engine/Extras/ushell/channels/unreal/core/cmds/uat.py
--- a/Engine/Extras/ushell/channels/unreal/core/cmds/uat.py
+++ b/Engine/Extras/ushell/channels/unreal/core/cmds/uat.py
@@ -53,13 +53,60 @@ class Uat(unrealcmd.MultiPlatformCmd):
         for script_root in uat_script_roots:
             yield from _read_foreach_subdir(script_root, _find_automation_rules)
 
+    def _get_build_script(self, name="RunUAT"):
+        ue_context = self.get_unreal_context()
+        script = ue_context.get_engine().get_dir() / "Build/BatchFiles"
+        script /= name + (".bat" if os.name == "nt" else ".sh")
+        return script
+
+    def _compile(self):
+        exec_context = self.get_exec_context()
+        cmd = self._get_build_script()
+        cmd = exec_context.create_runnable(cmd, "-Compile")
+        return cmd.run()
+
+    def _build_debug(self):
+        import subprocess as sp
+        from pathlib import Path
+        ue_context = self.get_unreal_context()
+        cmd = (
+            self._get_build_script("GetDotnetPath"),
+            "&&"
+            "where",
+            "dotnet",
+        )
+        dotnet_path = None
+        with sp.Popen(cmd, stdout=sp.PIPE) as proc:
+            for line in proc.stdout:
+                path = Path(line.decode().strip())
+                if path.stem == "dotnet" and path.is_file():
+                    dotnet_path = dotnet_path or path
+
+        cmd = (
+            "dotnet",
+            "build",
+            str(ue_context.get_engine().get_dir() / "Source/Programs/AutomationTool/AutomationTool.csproj"),
+            "-c", "Debug"
+        )
+        with sp.Popen(cmd, stdout=sp.PIPE) as proc:
+            for line in proc.stdout:
+                print(line.decode(), end="")
+        return proc.returncode
+
+    def _build(self):
+        if self.args.debug:
+            return self._build_debug()
+
+        exec_context = self.get_exec_context()
+        cmd = self._get_build_script("BuildUAT")
+        cmd = exec_context.create_runnable(cmd)
+        return cmd.run()
+
     def main(self):
         self.use_all_platforms()
 
         ue_context = self.get_unreal_context()
 
-        cmd = ue_context.get_engine().get_dir() / "Build/BatchFiles"
-        cmd /= "RunUAT.bat" if os.name == "nt" else "RunUAT.sh"
         args = (self.args.command,)
 
         if project := ue_context.get_project():
@@ -70,5 +117,6 @@ class Uat(unrealcmd.MultiPlatformCmd):
                     args = (*args, "-ScriptsForProject=" + project.get_name())
 
         exec_context = self.get_exec_context()
+        cmd = self._get_build_script()
         cmd = exec_context.create_runnable(cmd, *args, *self.args.uatargs)
         return cmd.run()
diff --git a/Engine/Extras/ushell/channels/unreal/core/pylib/unreal/_platform.py b/Engine/Extras/ushell/channels/unreal/core/pylib/unreal/_platform.py
--- a/Engine/Extras/ushell/channels/unreal/core/pylib/unreal/_platform.py
+++ b/Engine/Extras/ushell/channels/unreal/core/pylib/unreal/_platform.py
@@ -68,6 +68,7 @@ class Platform(object):
         return getattr(type(self), "config_name", None) or self.get_name()
 
     def get_autosdk_name(self):
+        self.get_version() # we're cheating; see SDK.json loading path
         return type(self).autosdk_name if hasattr(self, "autosdk_name") else self.get_name()
 
     def get_unreal_context(self):
@@ -104,12 +105,18 @@ class Platform(object):
                 except:
                     return
 
+            # patch the AutoSDK name while we're here...
+            if autosdk_name := sdk_info.get("AutoSDKPlatform", None):
+                type(self).autosdk_name = autosdk_name
+
+            # fetch the version value
             if version := sdk_info.get("AutoSDKDirectory", None):
                 return version
 
             if version := sdk_info.get("MainVersion", None):
                 return version
 
+            # ...or traverse up the platform forest
             if next_path := sdk_info.get("ParentSDKFile", None):
                 return read_sdk_json(path.parent / next_path)
 
diff --git a/Engine/Extras/ushell/channels/unreal/core/pylib/unreal/platforms/win64/linux.py b/Engine/Extras/ushell/channels/unreal/core/pylib/unreal/platforms/win64/linux.py
--- a/Engine/Extras/ushell/channels/unreal/core/pylib/unreal/platforms/win64/linux.py
+++ b/Engine/Extras/ushell/channels/unreal/core/pylib/unreal/platforms/win64/linux.py
@@ -10,5 +10,5 @@ class Platform(linux.PlatformBase):
         value = os.getenv(env_var)
         if not value:
             version = self.get_version()
-            value = f"Linux_x64/{version}/"
+            value = f"{self.get_autosdk_name()}/{version}/"
         yield env_var, value
