Updated ushell [40989220]

#rnx

diff --git a/Engine/Extras/ushell/.version b/Engine/Extras/ushell/.version
--- a/Engine/Extras/ushell/.version
+++ b/Engine/Extras/ushell/.version
@@ -1 +1 @@
-40827618
\ No newline at end of file
+40989220
\ No newline at end of file
diff --git a/Engine/Extras/ushell/channels/flow/core/system/flow/cmd.py b/Engine/Extras/ushell/channels/flow/core/system/flow/cmd.py
--- a/Engine/Extras/ushell/channels/flow/core/system/flow/cmd.py
+++ b/Engine/Extras/ushell/channels/flow/core/system/flow/cmd.py
@@ -206,6 +206,12 @@ class Cmd(_flick.Cmd, _ArgOverrider):
         from . import native as flow_native
         return flow_native.isatty(0) and flow_native.isatty(1) # 0=in, 1=out
 
+    def get_os_name(self):
+        if os.name == "nt":          return "nt"
+        if sys.platform == "darwin": return "mac"
+        if sys.platform == "linux":  return "linux"
+        return "unknown"
+
     def get_channel(self):
         return self._channel
 
diff --git a/Engine/Extras/ushell/channels/unreal/core/cmds/getbuild.py b/Engine/Extras/ushell/channels/unreal/core/cmds/getbuild.py
--- a/Engine/Extras/ushell/channels/unreal/core/cmds/getbuild.py
+++ b/Engine/Extras/ushell/channels/unreal/core/cmds/getbuild.py
@@ -12,14 +12,16 @@ import fzf
 
 #-------------------------------------------------------------------------------
 class GetBuild(unrealcmd.Cmd):
-    """ Download a build/artifact from UE Cloud
+    """ Download a build/artifact from UE Cloud Storage
 
 Unless otherwise specified it is downloaded to <projectdir>/Staged/<platform>.
 
-By default the latest build is downloaded. You can choose a specific build via --match as follows:
+If your search only matches one build that build is downloaded for you, to add filtering you can use --match as follows:
     * --match 8v93f39cvn3           - If you know the build id:
     * --match 1337                  - To search for the build cloest to a changelist
-    * --match job=v99f3;step=39f983 - For advanced searching you can specify key values"""
+    * --match job=v99f3;step=39f983 - For advanced searching you can specify key values
+
+If you find multiple builds that match you will be given the option to pick one, were the * indicates the build closest to your current commit."""
 
     buildtype   = unrealcmd.Arg(str, "The type of build to download. Usually packaged or staged builds.")
     platforms   = unrealcmd.Arg([str], "The platform(s) you want to download from a build")
@@ -129,8 +131,9 @@ By default the latest build is downloaded. You can choose a specific build via -
         server_platforms = ("win64", "linux", "mac")
         for platform in unrealcmd.Cmd.complete_platform(self, prefix):
             yield platform
+            yield platform + "-client"
             if platform in server_platforms:
-                yield platform + ".server"
+                yield platform + "-server"
 
     def complete_dest(self, prefix):
         prefix = prefix or "."
@@ -139,11 +142,25 @@ By default the latest build is downloaded. You can choose a specific build via -
                 yield item.name + "/"
 
     def convert_platform(self, platform):
+        wasClient = False
+        wasServer = False
+        if platform.endswith("-client"):
+            platform = platform[0:-7]
+            wasClient = True
+        if platform.endswith("-server"):
+            platform = platform[0:-7]
+            wasServer = True
+
         try:
-            return self.get_platform(platform).get_config_name()
+            convertedPlatform = self.get_platform(platform).get_config_name()
+            if wasClient:
+                return convertedPlatform + "-client"
+            if wasServer:
+                return convertedPlatform + "-server"
+            return convertedPlatform
         except ValueError:
             # unknown platform, just assume the name it valid
-            return platform 
+            return platform
 
     def convert_buildtype(self, build_type):
         if build_type == "packaged":
diff --git a/Engine/Extras/ushell/channels/unreal/core/cmds/run.py b/Engine/Extras/ushell/channels/unreal/core/cmds/run.py
--- a/Engine/Extras/ushell/channels/unreal/core/cmds/run.py
+++ b/Engine/Extras/ushell/channels/unreal/core/cmds/run.py
@@ -1,7 +1,6 @@
 # Copyright Epic Games, Inc. All Rights Reserved.
 
 import os
-import sys
 import unreal
 import unrealcmd
 import unreal.cmdline
@@ -186,7 +185,7 @@ class _ProgramTarget(_RunCmd, _Attachable, unrealcmd.MultiPlatformCmd):
             return
 
         launch_kwargs = {}
-        if not sys.stdout.isatty():
+        if not self.is_interactive():
             launch_kwargs = { "silent" : True }
         exec_context = self.get_exec_context()
         runnable = exec_context.create_runnable(str(binary_path), *args);
