Updated ushell [41141352]
#rnx

diff --git a/Engine/Extras/ushell/.version b/Engine/Extras/ushell/.version
--- a/Engine/Extras/ushell/.version
+++ b/Engine/Extras/ushell/.version
@@ -1 +1 @@
-40989220
\ No newline at end of file
+41141352
\ No newline at end of file
diff --git a/Engine/Extras/ushell/channels/unreal/core/cmds/insights.py b/Engine/Extras/ushell/channels/unreal/core/cmds/insights.py
--- a/Engine/Extras/ushell/channels/unreal/core/cmds/insights.py
+++ b/Engine/Extras/ushell/channels/unreal/core/cmds/insights.py
@@ -1,6 +1,7 @@
 import unreal
 import unrealcmd
 import subprocess as sp
+from pathlib import Path
 
 #-------------------------------------------------------------------------------
 class Insights(unrealcmd.Cmd):
@@ -48,6 +49,8 @@ class Insights(unrealcmd.Cmd):
         if self.args.trace == "latest":
             if latest := self._get_lastest_trace():
                 pass_through = (latest, *pass_through)
+        elif Path(self.args.trace).is_file():
+            pass_through = (self.args.trace, *pass_through)
 
         # Yup, this is mostly a shim around other ushell commands!
         ui_args = (
diff --git a/Engine/Extras/ushell/channels/unreal/core/cmds/stage.py b/Engine/Extras/ushell/channels/unreal/core/cmds/stage.py
--- a/Engine/Extras/ushell/channels/unreal/core/cmds/stage.py
+++ b/Engine/Extras/ushell/channels/unreal/core/cmds/stage.py
@@ -11,6 +11,7 @@ class _Impl(unrealcmd.Cmd):
     variant  = unrealcmd.Arg("development", "Build variant to stage binaries to copy")
     style    = unrealcmd.Arg("auto", "Type of result to stage")
     uatargs  = unrealcmd.Arg([str], "Additional arguments to pass to UAT")
+    debug    = unrealcmd.Opt(False, "Use a debug variant of UAT")
 
     complete_target = ("game", "client", "server")
 
@@ -135,6 +136,7 @@ class _Impl(unrealcmd.Cmd):
         dot_uat_cmd = "_uat"
         dot_uat_args = (
             "BuildCookRun",
+            "--debug" if self.args.debug else None,
             "--", # NB. NOT a UAT arg!
             #"-project=" + ue_context.get_project().get_path() <-- implicit via _uat
             target_arg,
diff --git a/Engine/Extras/ushell/channels/unreal/core/cmds/uat.py b/Engine/Extras/ushell/channels/unreal/core/cmds/uat.py
--- a/Engine/Extras/ushell/channels/unreal/core/cmds/uat.py
+++ b/Engine/Extras/ushell/channels/unreal/core/cmds/uat.py
@@ -11,6 +11,7 @@ class Uat(unrealcmd.MultiPlatformCmd):
     uatargs     = unrealcmd.Arg([str], "Arguments to pass to the UAT command")
     unprojected = unrealcmd.Opt(False, "No implicit '-project=' please")
     allscripts  = unrealcmd.Opt(False, "Ask UAT to compile all scripts (default is project-only)");
+    debug       = unrealcmd.Opt(False, "Use the debug variant of AutomationTool");
 
     def complete_command(self, prefix):
         import re
@@ -81,17 +82,20 @@ class Uat(unrealcmd.MultiPlatformCmd):
                 path = Path(line.decode().strip())
                 if path.stem == "dotnet" and path.is_file():
                     dotnet_path = dotnet_path or path
+        dotnet_path or dotnet_path or "dotnet"
 
         cmd = (
-            "dotnet",
+            dotnet_path,
             "build",
             str(ue_context.get_engine().get_dir() / "Source/Programs/AutomationTool/AutomationTool.csproj"),
             "-c", "Debug"
         )
-        with sp.Popen(cmd, stdout=sp.PIPE) as proc:
-            for line in proc.stdout:
-                print(line.decode(), end="")
-        return proc.returncode
+        exec_context = self.get_exec_context()
+        cmd = exec_context.create_runnable(*cmd)
+        for line in cmd:
+            print(line)
+
+        return cmd.get_return_code()
 
     def _build(self):
         if self.args.debug:
@@ -102,9 +106,26 @@ class Uat(unrealcmd.MultiPlatformCmd):
         cmd = exec_context.create_runnable(cmd)
         return cmd.run()
 
+    def _get_binary_path(self):
+        unconventional = self.args.debug
+        if not unconventional:
+            return self._get_build_script()
+
+        variant = "Debug"
+        ue_context = self.get_unreal_context()
+        engine_dir = ue_context.get_engine().get_dir()
+        return engine_dir / f"Source/Programs/AutomationTool/bin/{variant}/AutomationTool.exe"
+
     def main(self):
+        if self.get_os_name() != "nt":
+            self.print_warning("Disabling --debug as I haven't tried non-NT platforms")
+            self.args.debug = False
+
         self.use_all_platforms()
 
+        if ret := self._build():
+            return ret
+
         ue_context = self.get_unreal_context()
 
         args = (self.args.command,)
@@ -117,6 +138,6 @@ class Uat(unrealcmd.MultiPlatformCmd):
                     args = (*args, "-ScriptsForProject=" + project.get_name())
 
         exec_context = self.get_exec_context()
-        cmd = self._get_build_script()
+        cmd = self._get_binary_path()
         cmd = exec_context.create_runnable(cmd, *args, *self.args.uatargs)
         return cmd.run()
diff --git a/Engine/Extras/ushell/channels/unreal/perforce/cmds/clean.py b/Engine/Extras/ushell/channels/unreal/perforce/cmds/clean.py
--- a/Engine/Extras/ushell/channels/unreal/perforce/cmds/clean.py
+++ b/Engine/Extras/ushell/channels/unreal/perforce/cmds/clean.py
@@ -152,7 +152,8 @@ class _Cleaner(object):
         print(f"Moving directories to {self._rubbish_dir.name} and removing unversioned files")
 
     def remove_dir(self, path):
-        dest_name = "%08x_%016x_%s_%s" % (os.getpid(), id(path), path.parent.name, path.name)
+        path = path.absolute()
+        dest_name = "%08x_%016x_%s_%s" % (os.getpid(), hash(path), path.parent.name, path.name)
         try:
             path.rename(str(self._rubbish_dir / dest_name))
             self._dir_count += 1
