Updated ushell [41197121]

#rnx

diff --git a/Engine/Extras/ushell/.version b/Engine/Extras/ushell/.version
--- a/Engine/Extras/ushell/.version
+++ b/Engine/Extras/ushell/.version
@@ -1 +1 @@
-41141352
\ No newline at end of file
+41197121
\ No newline at end of file
diff --git a/Engine/Extras/ushell/channels/unreal/core/pylib/unreal/zencmd.py b/Engine/Extras/ushell/channels/unreal/core/pylib/unreal/zencmd.py
--- a/Engine/Extras/ushell/channels/unreal/core/pylib/unreal/zencmd.py
+++ b/Engine/Extras/ushell/channels/unreal/core/pylib/unreal/zencmd.py
@@ -10,6 +10,7 @@ import http.client
 import json
 import hashlib
 import tempfile
+import socket
 
 #-------------------------------------------------------------------------------
 def _get_ips_socket(_target_platform):
@@ -351,11 +352,13 @@ class ZenBaseCmd(unrealcmd.Cmd):
         target_gcmarker_dir = os.path.dirname(target_gcmarker)
         if not os.path.exists(target_gcmarker_dir):
             os.makedirs(target_gcmarker_dir)
+        remotehostnames = _get_ips(snapshot['targetplatform'])
+        remotehostnames.append("hostname://" + socket.getfqdn())
         projectstore = {
             "zenserver" : {
                 "islocalhost" : True,
                 "hostname" : "[::1]",
-                "remotehostnames" : _get_ips(snapshot['targetplatform']),
+                "remotehostnames" : remotehostnames,
                 "hostport" : 8558,
                 "projectid" : target_projectid,
                 "oplogid" : target_oplogid
