Updated ushell [41353259]
#rnx

diff --git a/Engine/Extras/ushell/.version b/Engine/Extras/ushell/.version
--- a/Engine/Extras/ushell/.version
+++ b/Engine/Extras/ushell/.version
@@ -1 +1 @@
-41197121
\ No newline at end of file
+41353259
\ No newline at end of file
diff --git a/Engine/Extras/ushell/channels/unreal/core/cmds/getbuild.py b/Engine/Extras/ushell/channels/unreal/core/cmds/getbuild.py
--- a/Engine/Extras/ushell/channels/unreal/core/cmds/getbuild.py
+++ b/Engine/Extras/ushell/channels/unreal/core/cmds/getbuild.py
@@ -2,6 +2,7 @@
 
 import sys
 import unrealcmd
+import unreal
 import tempfile
 import json
 import os
@@ -67,12 +68,12 @@ If you find multiple builds that match you will be given the option to pick one,
                 host=  str(value)
 
         # if host is a semicolon seperated list we parse out the first entry as zen cli does not yet support a list of hosts
-        if host!= None and host.find(';') != -1:
-            host = host.split(';')[0]
+        if host!= None and host.find(";") != -1:
+            host = host.split(";")[0]
 
         return host
 
-    def refresh_access_token(self):
+    def get_oidc_token_path(self):
         ue_context = self.get_unreal_context()
         engine = ue_context.get_engine()
 
@@ -89,40 +90,7 @@ If you find multiple builds that match you will be given the option to pick one,
         if not bin_path.is_file():
             raise FileNotFoundError(f"Unable to find '{bin_path}'")
 
-        service_name = self.lookup_service_name()
-        if not service_name:
-            raise ValueError("Unable to discover service name")
-
-        token_dir = tempfile.TemporaryDirectory()
-        token_file_path = f"{token_dir.name}/oidctoken.json"
-        oidcargs = (
-            "--ResultToConsole=true",
-            "--Service=" + str(service_name),
-            f"--OutFile={token_file_path}"
-        )
-
-        if project := ue_context.get_project():
-            oidcargs = (*oidcargs, "--Project=" + str(project.get_path()))
-
-        cmd = self.get_exec_context().create_runnable(bin_path, *oidcargs)
-        ret, output = cmd.run2()
-        if ret != 0:
-            return False
-
-        tokenresponse = None
-        with open(token_file_path, "r") as token_file:
-            tokenresponse = json.load(token_file)
-
-        self._AccessToken = tokenresponse['Token']
-        self.print_info("Cloud access token obtained.")
-
-        return True
-
-    def get_exec_context(self):
-        context = super().get_exec_context()
-        if hasattr(self, '_AccessToken') and self._AccessToken:
-            context.get_env()["CloudAccessToken"] = self._AccessToken
-        return context
+        return str(bin_path)
 
     def complete_buildtype(self, prefix):
         return ["packaged", "staged"]
@@ -184,7 +152,7 @@ If you find multiple builds that match you will be given the option to pick one,
             raise EnvironmentError("Unable to get Perforce client name")
         print("Using stream", branch)
         # massage the perforce stream name to match expectations, no initial / and - seperation between depot and stream
-        return branch.lstrip('/').replace('/', '-')
+        return branch.lstrip("/").replace("/", "-")
 
     def get_build_id(self, build_id_arg):
         if len(build_id_arg) == 24 and all(char in string.hexdigits for char in build_id_arg):
@@ -199,9 +167,9 @@ If you find multiple builds that match you will be given the option to pick one,
             search_request = "" # we special case handle the search for changelists
 
         query = {}
-        for statement in search_request.split(';'):
-            if statement.find('=') != -1:
-                k, v = statement.split('=')
+        for statement in search_request.split(";"):
+            if statement.find("=") != -1:
+                k, v = statement.split("=")
                 query[k] = {"$eq": v}
 
         if search_changelist:
@@ -211,27 +179,41 @@ If you find multiple builds that match you will be given the option to pick one,
 
     def search_build(self, host, namespace, bucket, query):
         # Run a search against the api to find builds that matches the request
-        url = f"{host}/api/v2/builds/{namespace}/{bucket}/search"
-        headers = {
-            "Authorization" : f"Bearer {self._AccessToken}",
-            "Content-Type" : "application/json"
-        }
+        temp_dir = tempfile.TemporaryDirectory()
+        query_file_path = f"{temp_dir.name}/query.json"
+        with open(query_file_path, "w") as query_file:
+            json.dump(query, query_file)
+
+        result_file_path = f"{temp_dir.name}/result.json"
+
+        args = (
+            "builds",
+            "list",
+            "--host=" + str(host),
+            "--namespace=" + str(namespace),
+            "--bucket=" + bucket,
+            "--oidctoken-exe-path=" + self.get_oidc_token_path(),
+            "--query-path=" + query_file_path,
+            "--result-path=" + result_file_path
+        )
+
+        result = self.run_zen(args)
+        if result != 0:
+            raise Exception("Failed to search")
 
-        data = json.dumps({"query": query}).encode()
-        request = urllib.request.Request(url, data=data, headers=headers)
-        response = urllib.request.urlopen(request)
-        body = response.read().decode()
+        with open(result_file_path, "r") as result_fd:
+            results = json.load(result_fd)["results"]
 
-        results = json.loads(body)["results"]
-        return [{'name' : r['metadata']['name'], 'platforms': {r['metadata']['platform']: r['buildId']}, 'commit': r['metadata'].get('commit', '')} for r in results]
+        return [{"name" : r["metadata"]["name"], "platforms": {r["metadata"]["platform"]: r["buildId"]}, "commit": r["metadata"].get("commit", "")} for r in results]
 
-    def _show_select_build_prompt(self, branch_cl, available_builds):
+    def _show_select_build_prompt(self, branch_cl, builds):
         closest_cl = 1 << 30
         p4_change_found = False
         active_platforms = set()
+        available_builds = builds.values()
         for build in available_builds:
-            cl = build['commit']
-            active_platforms.update(build['platforms'].keys())
+            cl = build["commit"]
+            active_platforms.update(build["platforms"].keys())
 
             if type(cl) is int or (type(cl) is str and cl.isnumeric()):
                 p4_change_found = True
@@ -247,9 +229,9 @@ If you find multiple builds that match you will be given the option to pick one,
 
         def list_available(closest_cl, available_builds):
             for build in available_builds:
-                cl = build['commit']
-                name = build['name']
-                platforms = build['platforms'].keys()
+                cl = build["commit"]
+                name = build["name"]
+                platforms = build["platforms"].keys()
                 middle = ""
                 for column_platform,short_name in columns.items():
                     middle += " "
@@ -261,11 +243,39 @@ If you find multiple builds that match you will be given the option to pick one,
                 yield "%10s%s %s" % (cl, middle, name)
 
         for line in fzf.run(list_available(closest_cl, available_builds), sort=False):
-            return line.split()[-1]
+            segments = line.split()
+            build_name = segments[-1]
+            commit = segments[0]
+            build_key = build_name + commit
+
+            if build_key in builds:
+                return build_key
+            elif build_name in builds:
+                return build_name
+            else:
+                raise Exception(f"Unable to map choice {line} back to selected build")
 
-    def main(self):
+    def run_zen(self, args):
         ue_context = self.get_unreal_context()
         engine = ue_context.get_engine()
+
+        bin_ext = ".exe" if os.name == "nt" else ""
+        bin_type = unreal.Platform.get_host()
+
+        bin_path = engine.get_dir()
+        bin_path /= f"Binaries/{bin_type}/zen{bin_ext}"
+
+        if not bin_path.is_file():
+            raise FileNotFoundError(f"Unable to find '{bin_path}'")
+
+        args = (*args, *self.args.runargs)
+        cmd = self.get_exec_context().create_runnable(bin_path, *args)
+        result = cmd.run()
+
+        return result
+
+    def main(self):
+        ue_context = self.get_unreal_context()
         project = self.args.project
         project_path = None
         if not project:
@@ -297,10 +307,6 @@ If you find multiple builds that match you will be given the option to pick one,
         else:
             host = self.lookup_host()
 
-        if not self.refresh_access_token():
-            self.print_error("Failed to acquire a access token")
-            return False
-
         if not host:
             self.print_error("No host specified")
             return False
@@ -333,11 +339,19 @@ If you find multiple builds that match you will be given the option to pick one,
                 builds = self.search_build(host, namespace, bucket, query)
                 print(f"Found {len(builds)} builds")
                 for build in builds:
-                    if build['name'] in found_builds:
-                        # this build already existed but a new platform was found add it to the list
-                        found_builds[build['name']]['platforms'].update(build['platforms'])
+                    build_name = build["name"]
+                    commit_str = str(build.get("commit", ""))
+                    build_key = build_name + commit_str
+                    if build_name in found_builds:
+                        # this build already existed but a new platform was found add it to the list if it didn't already exist
+                        for platform in build["platforms"].keys():
+                            if platform not in found_builds[build_key]["platforms"]:
+                                found_builds[build_key]["platforms"][platform] = build["platforms"][platform]
+                            else:
+                                self.print_info(f"Multiple builds for platform {platform} found for build {build_name}. Using newer build.")
+
                     else:
-                        found_builds[build['name']] = build
+                        found_builds[build_key] = build
 
             # remove builds that only have some of the platforms we want
             found_builds = {k:v for k, v in found_builds.items() if len(v["platforms"]) == len(platforms_to_download)}
@@ -354,34 +368,20 @@ If you find multiple builds that match you will be given the option to pick one,
             else:
                 # Ask the user to select a build
                 print("Select build to fetch:", end="")
-                selected_build_name = self._show_select_build_prompt(branch_cl, found_builds.values())
-                if not selected_build_name:
+                selected_build_id = self._show_select_build_prompt(branch_cl, found_builds)
+                if not selected_build_id:
                     self.print_error("No build was selected")
                     return False
 
+                selected_build = found_builds[selected_build_id]
+                selected_build_name = selected_build["name"]
                 print(f"Using build {selected_build_name}")
-                selected_build = found_builds[selected_build_name]
 
             # build the selection of build ids to download from which buckets
-            for platform, build_id in selected_build['platforms'].items():
+            for platform, build_id in selected_build["platforms"].items():
                 bucket = f"{project}.{build_type}.{branch}.{platform}".lower()
                 builds_to_download.extend([(bucket, build_id, platform)])
 
-        access_token_env_var = "CloudAccessToken"
-
-        bin_type = "Win64"
-        if sys.platform == "darwin": bin_type = "Mac"
-        elif sys.platform == "linux": bin_type = "Linux"
-
-        bin_ext = ".exe"
-        if sys.platform == "darwin": bin_ext = ""
-        elif sys.platform == "linux": bin_ext = ""
-
-        bin_path = engine.get_dir()
-        bin_path /= f"Binaries/{bin_type}/zen{bin_ext}"
-        if not bin_path.is_file():
-            raise FileNotFoundError(f"Unable to find '{bin_path}'")
-
         result = 0
         for bucket, build_id, build_platform in builds_to_download:
             if result:
@@ -392,16 +392,14 @@ If you find multiple builds that match you will be given the option to pick one,
             args = (
                 "builds",
                 "download",
-                "--url=" + str(host),
+                "--host=" + str(host),
                 "--namespace=" + str(namespace),
                 "--bucket=" + bucket,
-                "--access-token-env=" + access_token_env_var,
+                "--oidctoken-exe-path=" + self.get_oidc_token_path(),
                 build_dest,
                 build_id
             )
 
-            args = (*args, *self.args.runargs)
-            cmd = self.get_exec_context().create_runnable(bin_path, *args)
-            result = cmd.run()
+            result = self.run_zen(args)
 
         return result
