Updated ushell [41410096]
#rnx

diff --git a/Engine/Extras/ushell/.version b/Engine/Extras/ushell/.version
--- a/Engine/Extras/ushell/.version
+++ b/Engine/Extras/ushell/.version
@@ -1 +1 @@
-41353259
\ No newline at end of file
+41410096
\ No newline at end of file
diff --git a/Engine/Extras/ushell/channels/unreal/core/cmds/getbuild.py b/Engine/Extras/ushell/channels/unreal/core/cmds/getbuild.py
--- a/Engine/Extras/ushell/channels/unreal/core/cmds/getbuild.py
+++ b/Engine/Extras/ushell/channels/unreal/core/cmds/getbuild.py
@@ -245,7 +245,7 @@ If you find multiple builds that match you will be given the option to pick one,
         for line in fzf.run(list_available(closest_cl, available_builds), sort=False):
             segments = line.split()
             build_name = segments[-1]
-            commit = segments[0]
+            commit = segments[0].rstrip("*")
             build_key = build_name + commit
 
             if build_key in builds:
