Updated ushell [41484606]

#rnx

diff --git a/Engine/Extras/ushell/.version b/Engine/Extras/ushell/.version
--- a/Engine/Extras/ushell/.version
+++ b/Engine/Extras/ushell/.version
@@ -1 +1 @@
-41410096
\ No newline at end of file
+41484606
\ No newline at end of file
diff --git a/Engine/Extras/ushell/channels/unreal/core/cmds/gather.py b/Engine/Extras/ushell/channels/unreal/core/cmds/gather.py
--- a/Engine/Extras/ushell/channels/unreal/core/cmds/gather.py
+++ b/Engine/Extras/ushell/channels/unreal/core/cmds/gather.py
@@ -4,7 +4,7 @@ from pathlib import Path
 
 #-------------------------------------------------------------------------------
 def _find_ushell_dir(start_path):
-    start_path = Path(start_path) or Path(__file__)
+    start_path = (Path(start_path) / "x") or Path(__file__)
     for item in start_path.parents:
         if (item / ("ushell.bat")).is_file():
             return item;
diff --git a/Engine/Extras/ushell/channels/unreal/core/cmds/uat.py b/Engine/Extras/ushell/channels/unreal/core/cmds/uat.py
--- a/Engine/Extras/ushell/channels/unreal/core/cmds/uat.py
+++ b/Engine/Extras/ushell/channels/unreal/core/cmds/uat.py
@@ -117,19 +117,16 @@ class Uat(unrealcmd.MultiPlatformCmd):
         return engine_dir / f"Source/Programs/AutomationTool/bin/{variant}/AutomationTool.exe"
 
     def main(self):
-        if self.get_os_name() != "nt":
-            self.print_warning("Disabling --debug as I haven't tried non-NT platforms")
-            self.args.debug = False
-
         self.use_all_platforms()
 
-        if ret := self._build():
-            return ret
-
-        ue_context = self.get_unreal_context()
+        if self.get_os_name() == "nt":
+            # POSIX's BuildUAT.sh script isn't self-contained and doesn't run alone
+            if ret := self._build():
+                return ret
 
         args = (self.args.command,)
 
+        ue_context = self.get_unreal_context()
         if project := ue_context.get_project():
             if args[0] and not self.args.unprojected:
                 args = (*args, "-project=" + str(project.get_path()))
