Updated ushell [41487163]
#rnx

diff --git a/Engine/Extras/ushell/.version b/Engine/Extras/ushell/.version
--- a/Engine/Extras/ushell/.version
+++ b/Engine/Extras/ushell/.version
@@ -1 +1 @@
-41484606
\ No newline at end of file
+41487163
\ No newline at end of file
diff --git a/Engine/Extras/ushell/channels/unreal/core/pylib/unreal/platforms/mac/mac.py b/Engine/Extras/ushell/channels/unreal/core/pylib/unreal/platforms/mac/mac.py
--- a/Engine/Extras/ushell/channels/unreal/core/pylib/unreal/platforms/mac/mac.py
+++ b/Engine/Extras/ushell/channels/unreal/core/pylib/unreal/platforms/mac/mac.py
@@ -2,6 +2,31 @@
 
 import os
 import unreal
+import subprocess as sp
+from pathlib import Path
+
+#-------------------------------------------------------------------------------
+def _is_sandboxed(binary_path):
+    cs_cmd = (
+        "codesign",
+        "-d",
+        "--xml",
+        "--entitlements=-",
+        binary_path,
+    )
+    with sp.Popen(cs_cmd, stdout=sp.PIPE, stderr=sp.DEVNULL) as proc:
+        it = iter(proc.stdout)
+        xml_data = b"".join(it)
+
+    import xml.etree.ElementTree as et
+    xml = et.fromstring(xml_data)
+    it = (x for x in xml.findall(".//dict/*"))
+    for item in it:
+        if item.text != "com.apple.security.app-sandbox":
+            continue
+
+        item = next(it)
+        return item.tag == "true"
 
 #-------------------------------------------------------------------------------
 class Platform(unreal.Platform):
@@ -26,6 +51,31 @@ class Platform(unreal.Platform):
         if target == "server": return "MacServer"
 
     def _launch(self, exec_context, stage_dir, binary_path, args):
+        print("-- Checking entitlements")
+        constrained = False
+        try:
+            if not os.getenv("USHELL_NO_BASEDIR", None):
+                constrained = _is_sandboxed(binary_path)
+        except IOError:
+            print("NOTE: failed calling 'codesign' to check for sandbox entitlement")
+
+        if constrained:
+            lhs = Path(stage_dir)
+            rhs = Path(binary_path).parent.parent
+            if not lhs.is_relative_to(rhs):
+                print("\x1b[91m!!")
+                print(f"!! Binary '{binary_path}' is sandboxed")
+                print(f"!! It cannot access dataset '{stage_dir}'")
+                print(f"!!")
+                print("!! BUILD BINARIES WITH '.build [...] -- -NoEntitlements'")
+                print("!!\x1b[0m")
+                import time
+                time.sleep(2)
+                stage_dir = None
+
+        if os.getenv("USHELL_NO_BASEDIR", None):
+            stage_dir = None
+
         if stage_dir:
             midfix = "Engine";
             if project := self.get_unreal_context().get_project():
