Updated ushell [41567270]

#rnx

diff --git a/Engine/Extras/ushell/.version b/Engine/Extras/ushell/.version
--- a/Engine/Extras/ushell/.version
+++ b/Engine/Extras/ushell/.version
@@ -1 +1 @@
-41487163
\ No newline at end of file
+41567270
\ No newline at end of file
diff --git a/Engine/Extras/ushell/channels/flow/core/system/flow/_runnable.py b/Engine/Extras/ushell/channels/flow/core/system/flow/_runnable.py
--- a/Engine/Extras/ushell/channels/flow/core/system/flow/_runnable.py
+++ b/Engine/Extras/ushell/channels/flow/core/system/flow/_runnable.py
@@ -91,7 +91,7 @@ class Runnable(object):
 
         kwargs = {
             "env" : self._env,
-            "stdin" : stdin,
+            "stdin" : None if stdin is None else sp.DEVNULL,
             "stdout" : sp.PIPE if stdout else None if stdout == None else sp.DEVNULL,
             "stderr" : sp.STDOUT if stdout else None,
             "creationflags" : create_flags,
diff --git a/Engine/Extras/ushell/channels/unreal/core/cmds/build.py b/Engine/Extras/ushell/channels/unreal/core/cmds/build.py
--- a/Engine/Extras/ushell/channels/unreal/core/cmds/build.py
+++ b/Engine/Extras/ushell/channels/unreal/core/cmds/build.py
@@ -213,6 +213,7 @@ class _BuildCmd(unrealcmd.Cmd):
     noxge     = unrealcmd.Opt(False, "Disables IncrediBuild")
     analyze   = unrealcmd.Opt(("", "visualcpp"), "Run static analysis (add =pvsstudio for PVS)")
     projected = unrealcmd.Opt(False, "Add the -Project= argument when running UnrealBuildTool")
+    noshowcfg = unrealcmd.Opt(False, "Bypass the printing of configuration from XML files")
 
     complete_analyze = ("visualcpp", "pvsstudio")
     #complete_ubtargs = ...see end of file
@@ -321,10 +322,9 @@ class _BuildCmd(unrealcmd.Cmd):
 
         return builder
 
-    def run_build(self, builder):
-        try: exec_context = self._exec_context
-        except AttributeError: self._exec_context = self.get_exec_context()
-        exec_context = self._exec_context
+    def show_build_config(self):
+        if not self.is_interactive() or self.args.noshowcfg:
+            return
 
         # Inform the user of the current build configuration
         self.print_info("Build configuration")
@@ -341,6 +341,13 @@ class _BuildCmd(unrealcmd.Cmd):
                 self.print_warning(str(e))
                 continue
 
+    def run_build(self, builder):
+        self.show_build_config()
+
+        try: exec_context = self._exec_context
+        except AttributeError: self._exec_context = self.get_exec_context()
+        exec_context = self._exec_context
+
         self.print_info("Running UnrealBuildTool")
         printer = _PrettyPrinter(self)
         try:
diff --git a/Engine/Extras/ushell/channels/unreal/core/cmds/getbuild.py b/Engine/Extras/ushell/channels/unreal/core/cmds/getbuild.py
--- a/Engine/Extras/ushell/channels/unreal/core/cmds/getbuild.py
+++ b/Engine/Extras/ushell/channels/unreal/core/cmds/getbuild.py
@@ -32,6 +32,8 @@ If you find multiple builds that match you will be given the option to pick one,
     branch      = unrealcmd.Opt("", "Branch/Stream to fetch data for, defaults to current p4 stream")
     namespace   = unrealcmd.Opt("", "The namespace to use overriding what is defined in engine ini")
     host        = unrealcmd.Opt("", "The host name of the server to connect to, defaults to your cloud storage server in engine ini")
+    clean       = unrealcmd.Opt(False, "Ignore any previously downloaded data and fetch it all again")
+    verbose     = unrealcmd.Opt(False, "Enable verbose logging output")
 
     dest        = unrealcmd.Opt("", "The directory to download the build to. Local contents will be overwritten. Defaults to `<projectdir>/Staged/<platform>`")
     match       = unrealcmd.Opt("", "Either a build id (hexadecimal string), a changelist (number) or a key value seperate list of key=values you want to find a build for. e.g. changelist=1337;job=foo")
@@ -177,7 +179,7 @@ If you find multiple builds that match you will be given the option to pick one,
 
         return query
 
-    def search_build(self, host, namespace, bucket, query):
+    def search_build(self, host, override_host, namespace, bucket, query):
         # Run a search against the api to find builds that matches the request
         temp_dir = tempfile.TemporaryDirectory()
         query_file_path = f"{temp_dir.name}/query.json"
@@ -189,7 +191,7 @@ If you find multiple builds that match you will be given the option to pick one,
         args = (
             "builds",
             "list",
-            "--host=" + str(host),
+            "--host=" + str(host) if host != None else "--override-host=" + str(override_host),
             "--namespace=" + str(namespace),
             "--bucket=" + bucket,
             "--oidctoken-exe-path=" + self.get_oidc_token_path(),
@@ -197,7 +199,7 @@ If you find multiple builds that match you will be given the option to pick one,
             "--result-path=" + result_file_path
         )
 
-        result = self.run_zen(args)
+        result = self.run_zen(args, include_passthrough_args = False)
         if result != 0:
             raise Exception("Failed to search")
 
@@ -255,7 +257,7 @@ If you find multiple builds that match you will be given the option to pick one,
             else:
                 raise Exception(f"Unable to map choice {line} back to selected build")
 
-    def run_zen(self, args):
+    def run_zen(self, args, include_passthrough_args = True):
         ue_context = self.get_unreal_context()
         engine = ue_context.get_engine()
 
@@ -268,7 +270,8 @@ If you find multiple builds that match you will be given the option to pick one,
         if not bin_path.is_file():
             raise FileNotFoundError(f"Unable to find '{bin_path}'")
 
-        args = (*args, *self.args.runargs)
+        if include_passthrough_args:
+            args = (*args, *self.args.runargs)
         cmd = self.get_exec_context().create_runnable(bin_path, *args)
         result = cmd.run()
 
@@ -302,12 +305,14 @@ If you find multiple builds that match you will be given the option to pick one,
         else:
             namespace = self.lookup_namespace()
 
+        host = None
+        override_host = None
         if self.args.host:
-            host = self.args.host
+            override_host = self.args.host
         else:
             host = self.lookup_host()
 
-        if not host:
+        if not host and not override_host:
             self.print_error("No host specified")
             return False
         if not namespace:
@@ -336,7 +341,7 @@ If you find multiple builds that match you will be given the option to pick one,
                 self.print_info(f"Enumerating builds for platform {platform}")
                 bucket = f"{project}.{build_type}.{branch}.{platform}".lower()
 
-                builds = self.search_build(host, namespace, bucket, query)
+                builds = self.search_build(host, override_host, namespace, bucket, query)
                 print(f"Found {len(builds)} builds")
                 for build in builds:
                     build_name = build["name"]
@@ -392,10 +397,12 @@ If you find multiple builds that match you will be given the option to pick one,
             args = (
                 "builds",
                 "download",
-                "--host=" + str(host),
+                "--host=" + str(host) if host != None else "--override-host=" + str(override_host),
                 "--namespace=" + str(namespace),
                 "--bucket=" + bucket,
                 "--oidctoken-exe-path=" + self.get_oidc_token_path(),
+                "--clean" if self.args.clean else "",
+                "--verbose" if self.args.verbose else "",
                 build_dest,
                 build_id
             )
diff --git a/Engine/Extras/ushell/channels/unreal/core/cmds/uat.py b/Engine/Extras/ushell/channels/unreal/core/cmds/uat.py
--- a/Engine/Extras/ushell/channels/unreal/core/cmds/uat.py
+++ b/Engine/Extras/ushell/channels/unreal/core/cmds/uat.py
@@ -137,4 +137,4 @@ class Uat(unrealcmd.MultiPlatformCmd):
         exec_context = self.get_exec_context()
         cmd = self._get_binary_path()
         cmd = exec_context.create_runnable(cmd, *args, *self.args.uatargs)
-        return cmd.run()
+        return cmd.run(stdin=False)
