Updated ushell [41716299]
#rnx

diff --git a/Engine/Extras/ushell/.version b/Engine/Extras/ushell/.version
--- a/Engine/Extras/ushell/.version
+++ b/Engine/Extras/ushell/.version
@@ -1 +1 @@
-41567270
\ No newline at end of file
+41716299
\ No newline at end of file
diff --git a/Engine/Extras/ushell/channels/unreal/core/cmds/uat.py b/Engine/Extras/ushell/channels/unreal/core/cmds/uat.py
--- a/Engine/Extras/ushell/channels/unreal/core/cmds/uat.py
+++ b/Engine/Extras/ushell/channels/unreal/core/cmds/uat.py
@@ -12,6 +12,7 @@ class Uat(unrealcmd.MultiPlatformCmd):
     unprojected = unrealcmd.Opt(False, "No implicit '-project=' please")
     allscripts  = unrealcmd.Opt(False, "Ask UAT to compile all scripts (default is project-only)");
     debug       = unrealcmd.Opt(False, "Use the debug variant of AutomationTool");
+    attach      = unrealcmd.Opt(False, "Open VS to debug UAT (not exactly 'attach')");
 
     def complete_command(self, prefix):
         import re
@@ -67,6 +68,8 @@ class Uat(unrealcmd.MultiPlatformCmd):
         return cmd.run()
 
     def _build_debug(self):
+        self.print_info("Building debug UAT")
+
         import subprocess as sp
         from pathlib import Path
         ue_context = self.get_unreal_context()
@@ -106,16 +109,32 @@ class Uat(unrealcmd.MultiPlatformCmd):
         cmd = exec_context.create_runnable(cmd)
         return cmd.run()
 
-    def _get_binary_path(self):
-        unconventional = self.args.debug
+    def _get_binary_path(self, *, variant=None):
+        unconventional = self.args.debug or variant
         if not unconventional:
             return self._get_build_script()
 
-        variant = "Debug"
+        variant = variant or "Debug"
         ue_context = self.get_unreal_context()
         engine_dir = ue_context.get_engine().get_dir()
         return engine_dir / f"Source/Programs/AutomationTool/bin/{variant}/AutomationTool.exe"
 
+    def _launch_debugger(self, *args):
+        self.print_info("Launching Visual Studio")
+
+        variant = "Debug" if self.args.debug else "Development"
+        cmd = self._get_binary_path(variant=variant)
+
+        if self.get_os_name() != "nt":
+            self.print_warning("Debugging unavailable. Printing command instead")
+            print(cmd, *args)
+            return False
+
+        exec_context = self.get_exec_context()
+        cmd = exec_context.create_runnable("devenv.exe", "/debugexe", cmd, *args)
+        cmd.launch()
+        print("\nRemember to change debugger type to '.NET Core'")
+
     def main(self):
         self.use_all_platforms()
 
@@ -134,7 +153,12 @@ class Uat(unrealcmd.MultiPlatformCmd):
                 if not self.args.allscripts:
                     args = (*args, "-ScriptsForProject=" + project.get_name())
 
+        args = (*args, *self.args.uatargs)
+
+        if self.args.attach:
+            return self._launch_debugger(*args)
+
         exec_context = self.get_exec_context()
         cmd = self._get_binary_path()
-        cmd = exec_context.create_runnable(cmd, *args, *self.args.uatargs)
+        cmd = exec_context.create_runnable(cmd, *args)
         return cmd.run(stdin=False)
