Updated ushell [41861897]

#rnx

diff --git a/Engine/Extras/ushell/.version b/Engine/Extras/ushell/.version
--- a/Engine/Extras/ushell/.version
+++ b/Engine/Extras/ushell/.version
@@ -1 +1 @@
-41716299
\ No newline at end of file
+41861897
\ No newline at end of file
diff --git a/Engine/Extras/ushell/channels/unreal/core/cmds/clangdb.py b/Engine/Extras/ushell/channels/unreal/core/cmds/clangdb.py
--- a/Engine/Extras/ushell/channels/unreal/core/cmds/clangdb.py
+++ b/Engine/Extras/ushell/channels/unreal/core/cmds/clangdb.py
@@ -36,13 +36,12 @@ class ClangDb(unrealcmd.Cmd):
     target    = unrealcmd.Arg("", "Target to use when generating the database")
     ubtargs   = unrealcmd.Arg([str], "Arguments to pass onwards to UnrealBuildTool")
     platform  = unrealcmd.Opt("", "Platform to build a database for")
+    projected = unrealcmd.Opt(False, "Add the -Project= argument when running UnrealBuildTool")
 
     @unrealcmd.Cmd.summarise
     def _run_ubt(self, target, platform, ubt_args):
         ue_context = self.get_unreal_context()
         ubt = ue_context.get_engine().get_ubt()
-        if project := ue_context.get_project():
-            ubt_args = ("-Project=" + str(project.get_path()), *ubt_args)
 
         exec_context = self.get_exec_context()
         for cmd, args in ubt.read_actions(target, platform, "Development", *ubt_args):
@@ -72,16 +71,21 @@ class ClangDb(unrealcmd.Cmd):
 
         # Calculate how to invoke UBT
         target = self.args.target
-        if not target:
+        if target:
+            target = ue_context.get_target_by_name(self.args.target)
+            if not target:
+                raise ValueError(f"Unable to find target {self.args.target}")
+        else:
             target = ue_context.get_target_by_type(unreal.TargetType.EDITOR)
-            target = target.get_name()
+            if not target:
+                raise ValueError(f"Unable to find default editor target")
 
         platform = unreal.Platform.get_host()
         if self.args.platform:
             platform = self.get_platform(self.args.platform).get_name()
 
         self.print_info("Database properties")
-        print("Target:", target)
+        print("Target:", target.get_name())
         print("Platform:", platform)
 
         ubt_args = (
@@ -90,9 +94,13 @@ class ClangDb(unrealcmd.Cmd):
             *self.args.ubtargs,
         )
 
+        if target.is_project_target() or self.args.projected:
+            if project := ue_context.get_project():
+                ubt_args = ("-Project=" + str(project.get_path()), *ubt_args)
+
         # Run UBT
         self.print_info("Running UnrealBuildTool")
-        ret = self._run_ubt(target, platform, ubt_args)
+        ret = self._run_ubt(target.get_name(), platform, ubt_args)
         if ret:
             return ret
 
diff --git a/Engine/Extras/ushell/channels/unreal/core/cmds/snapshot.py b/Engine/Extras/ushell/channels/unreal/core/cmds/snapshot.py
new file mode 100644
--- /dev/null
+++ b/Engine/Extras/ushell/channels/unreal/core/cmds/snapshot.py
@@ -0,0 +1,7 @@
+# Copyright Epic Games, Inc. All Rights Reserved.
+
+import zen.snapshot
+
+List = zen.snapshot.List
+Find = zen.snapshot.Find
+Get = zen.snapshot.Get
\ No newline at end of file
diff --git a/Engine/Extras/ushell/channels/unreal/core/cmds/zen.py b/Engine/Extras/ushell/channels/unreal/core/cmds/zen.py
--- a/Engine/Extras/ushell/channels/unreal/core/cmds/zen.py
+++ b/Engine/Extras/ushell/channels/unreal/core/cmds/zen.py
@@ -6,20 +6,20 @@ import shutil
 import unreal
 import unrealcmd
 import unreal.cmdline
-import unreal.zencmd
+import zen.cmd
 import http.client
 import json
 import hashlib
 
 #-------------------------------------------------------------------------------
-class Dashboard(unreal.zencmd.ZenUETargetBaseCmd):
+class Dashboard(zen.cmd.ZenUETargetBaseCmd):
     """ Starts the zen dashboard GUI."""
 
     def main(self):
         return self.run_ue_program_target("ZenDashboard")
 
 #-------------------------------------------------------------------------------
-class Start(unreal.zencmd.ZenUETargetBaseCmd):
+class Start(zen.cmd.ZenUETargetBaseCmd):
     """ Starts an instance of zenserver."""
     SponsorProcessID = unrealcmd.Opt("", "Process ID to be added as a sponsor process for the zenserver process")
 
@@ -36,28 +36,28 @@ class Start(unreal.zencmd.ZenUETargetBaseCmd):
         return self.run_ue_program_target("ZenLaunch", args)
 
 #-------------------------------------------------------------------------------
-class Stop(unreal.zencmd.ZenUtilityBaseCmd):
+class Stop(zen.cmd.ZenUtilityBaseCmd):
     """ Stops any running instance of zenserver."""
 
     def main(self):
         return self.run_zen_utility(["down"])
 
 #-------------------------------------------------------------------------------
-class Status(unreal.zencmd.ZenUtilityBaseCmd):
+class Status(zen.cmd.ZenUtilityBaseCmd):
     """ Get the status of running zenserver instances."""
 
     def main(self):
         return self.run_zen_utility(["status"])
 
 #-------------------------------------------------------------------------------
-class Version(unreal.zencmd.ZenUtilityBaseCmd):
+class Version(zen.cmd.ZenUtilityBaseCmd):
     """ Get the version of the in-tree zenserver executable."""
 
     def main(self):
         return self.run_zen_utility(["version"])
 
 #-------------------------------------------------------------------------------
-class ImportSnapshot(unreal.zencmd.ZenUtilityBaseCmd):
+class ImportSnapshot(zen.cmd.ZenUtilityBaseCmd):
     """ Imports an oplog snapshot into the running zenserver process."""
     snapshotdescriptor = unrealcmd.Arg(str, "Snapshot descriptor file path to import from")
     snapshotindex = unrealcmd.Arg(0, "0-based index of the snapshot within the snapshot descriptor to import from")
@@ -99,7 +99,7 @@ class ImportSnapshot(unreal.zencmd.ZenUtilityBaseCmd):
         return self.import_snapshot(snapshot, self.args.sourcehost, self.args.projectid, self.args.oplog, self.args.asyncimport, self.args.forceimport)
 
 #-------------------------------------------------------------------------------
-class CreateWorkspace(unreal.zencmd.ZenUtilityBaseCmd):
+class CreateWorkspace(zen.cmd.ZenUtilityBaseCmd):
     """ Create a Zen workspace"""
     base_dir = unrealcmd.Arg(str, "Base path for the Zen workspace")
     dynamic = unrealcmd.Opt(False, "Should the zen server allow for dynamic creation of shares by clients in this workspace")
@@ -115,7 +115,7 @@ class CreateWorkspace(unreal.zencmd.ZenUtilityBaseCmd):
             self.print_info("Workspace already exists at this location")
 
 #-------------------------------------------------------------------------------
-class CreateShare(unreal.zencmd.ZenUtilityBaseCmd):
+class CreateShare(zen.cmd.ZenUtilityBaseCmd):
     """ Create a Zen workspace share"""
     share_dir = unrealcmd.Arg(str, "Base path for the Zen share. Must be inside an already existing workspace")
 
diff --git a/Engine/Extras/ushell/channels/unreal/core/describe.flow.py b/Engine/Extras/ushell/channels/unreal/core/describe.flow.py
--- a/Engine/Extras/ushell/channels/unreal/core/describe.flow.py
+++ b/Engine/Extras/ushell/channels/unreal/core/describe.flow.py
@@ -43,6 +43,17 @@ if bundle := ripgrep.bundle(platform="darwin-arm64"):
 
 
 
+#-------------------------------------------------------------------------------
+vswhere = flow.describe.Tool()
+vswhere.version("3.1.7")
+vswhere.payload("[REDACTED_PROFILE_URL]/vswhere/releases/download/$VERSION/vswhere.exe")
+vswhere.sha1("e3fa9b2db259d8875170717469779ea1280c8466")
+vswhere.platform("win32")
+vswhere.bin("vswhere.exe")
+vswhere.source("[REDACTED_PROFILE_URL]/vswhere/releases", r"(\d+\.\d+\.\d+)/vswhere.exe")
+
+
+
 #-------------------------------------------------------------------------------
 build_target = flow.describe.Command()
 build_target.source("cmds/build.py", "Build")
@@ -274,6 +285,19 @@ perf_insights = flow.describe.Command()
 perf_insights.source("cmds/insights.py", "Insights")
 perf_insights.invoke("perf", "insights")
 
+#-------------------------------------------------------------------------------
+snapshot_get_cmd = flow.describe.Command()
+snapshot_get_cmd.source("cmds/snapshot.py", "Find")
+snapshot_get_cmd.invoke("zen", "snapshot", "find")
+
+snapshot_find_cmd = flow.describe.Command()
+snapshot_find_cmd.source("cmds/snapshot.py", "Get")
+snapshot_find_cmd.invoke("zen", "snapshot", "get")
+
+snapshot_list_cmd = flow.describe.Command()
+snapshot_list_cmd.source("cmds/snapshot.py", "List")
+snapshot_list_cmd.invoke("zen", "snapshot", "list")
+
 
 
 #-------------------------------------------------------------------------------
@@ -296,17 +320,6 @@ tips.prefix("$")
 
 
 
-#-------------------------------------------------------------------------------
-vswhere = flow.describe.Tool()
-vswhere.version("3.1.7")
-vswhere.payload("[REDACTED_PROFILE_URL]/vswhere/releases/download/$VERSION/vswhere.exe")
-vswhere.sha1("e3fa9b2db259d8875170717469779ea1280c8466")
-vswhere.platform("win32")
-vswhere.bin("vswhere.exe")
-vswhere.source("[REDACTED_PROFILE_URL]/vswhere/releases", r"(\d+\.\d+\.\d+)/vswhere.exe")
-
-
-
 #-------------------------------------------------------------------------------
 unreal = flow.describe.Channel()
 unreal.version("0")
diff --git a/Engine/Extras/ushell/channels/unreal/core/pylib/unreal/zencmd.py b/Engine/Extras/ushell/channels/unreal/core/pylib/unreal/zencmd.py
deleted file mode 100644
diff --git a/Engine/Extras/ushell/channels/unreal/core/pylib/zen/__init__.py b/Engine/Extras/ushell/channels/unreal/core/pylib/zen/__init__.py
new file mode 100644
--- /dev/null
+++ b/Engine/Extras/ushell/channels/unreal/core/pylib/zen/__init__.py
@@ -0,0 +1 @@
+# Copyright Epic Games, Inc. All Rights Reserved.
diff --git a/Engine/Extras/ushell/channels/unreal/core/pylib/zen/cmd.py b/Engine/Extras/ushell/channels/unreal/core/pylib/zen/cmd.py
new file mode 100644
--- /dev/null
+++ b/Engine/Extras/ushell/channels/unreal/core/pylib/zen/cmd.py
@@ -0,0 +1,590 @@
+# Copyright Epic Games, Inc. All Rights Reserved.
+
+import os
+import sys
+import shutil
+import unreal
+import unrealcmd
+import unreal.cmdline
+import http.client
+import json
+import hashlib
+import tempfile
+import socket
+
+#-------------------------------------------------------------------------------
+def _get_ips_socket(_target_platform):
+    import socket
+    def _get_ip(family, ip_addr):
+        with socket.socket(family, socket.SOCK_DGRAM) as s:
+            try:
+                s.connect((ip_addr, 1))
+                return s.getsockname()[0]
+            except Exception as e:
+                return
+    ret = (_get_ip(socket.AF_INET6, "fe80::0"), _get_ip(socket.AF_INET, "172.31.255.255"))
+    return [x for x in ret if x]
+
+def _get_ips_osx(target_platform):
+    ret = _get_ips_socket(target_platform)
+    try:
+        import subprocess
+        import re
+        # since grep will return an exit code of 1 if not found, force it to return 0, otherwise it throws an exception
+        segments = re.split("\\s+", subprocess.check_output('ifconfig | grep 169.254. || true', shell=True).decode('utf-8').strip())
+        print(segments)
+        # if you have multiple devices plugged in, this will iterate over the segments (every 6th item is the IP)
+        # ie: ['inet', '169.254.12.24', 'netmask', '0xffff0000', 'broadcast', '169.254.255.255', 'inet', '169.254.110.30', 'netmask', '0xffff0000', 'broadcast', '169.254.255.255']
+        ret.extend(segments[1::6])
+
+        # Can not rely on a stable server link local address above as each time USB cable is plugged in,
+        # the IP changes. So use the mac server's hostname for the client to connect against as a fallback.
+        # Note: only do this for members of the iOS family - and there's no guarantee the mac has a bonjour name
+        if any (x in target_platform.lower() for x in ["ios", "tvos", "vision"]):
+            if bonjourName := subprocess.check_output('scutil --get ComputerName', shell=True).decode('utf-8').strip():
+                ret.append("macserver://" + bonjourName)
+    except:
+        print("Exception thrown in _get_ips_osx")
+        pass
+    return ret
+
+def _get_ips_powershell():
+    import subprocess as sp
+    ps_cmd = (
+        "powershell.exe",
+        "-NoProfile",
+        "-NoLogo",
+        "-Command",
+        r'(Get-CimInstance -Query "SELECT * FROM Win32_NetworkAdapterConfiguration WHERE IPEnabled=true"|Sort -Property IPConnectionMetric).IPAddress',
+    )
+    startupinfo = sp.STARTUPINFO()
+    startupinfo.dwFlags = sp.STARTF_USESHOWWINDOW
+    startupinfo.wShowWindow = sp.SW_HIDE
+    with sp.Popen(ps_cmd, stdout=sp.PIPE, stderr=sp.DEVNULL, stdin=sp.DEVNULL, startupinfo=startupinfo, creationflags=sp.CREATE_NEW_CONSOLE) as proc:
+        ret = [x.strip().decode() for x in proc.stdout]
+        proc.stdout.close()
+    return ret if proc.returncode == 0 else None
+
+def _get_ips_wmic():
+    import subprocess as sp
+    ret = []
+    current_ips = None
+    wmic_cmd = ("wmic", "nicconfig", "where", "ipenabled=true", "get", "ipaddress,ipconnectionmetric", "/format:value")
+    with sp.Popen(wmic_cmd, stdout=sp.PIPE, stderr=sp.DEVNULL) as proc:
+        for line in proc.stdout:
+            if line.startswith(b"IPAddress="):
+                if line := line[10:].strip():
+                    current_ips = line
+            elif line.startswith(b"IPConnectionMetric="):
+                if line := line[19:].strip():
+                    connection_metric = int(line)
+                    for ip in eval(current_ips):
+                        ret.append((connection_metric, ip))
+    ret.sort(key=lambda p: p[0])
+    return [item[1] for item in ret]
+
+def _get_ips_nt(target_platform):
+    try:
+        ret = _get_ips_powershell()
+        if not ret:
+            ret = _get_ips_wmic()
+    except:
+        ret = _get_ips_socket(target_platform)
+    ret.sort(key=lambda x: "." in x) # sort ipv4 last
+    return ret
+
+_get_ips = _get_ips_nt if os.name == "nt" else (_get_ips_osx if sys.platform == "darwin" else _get_ips_socket)
+
+#-------------------------------------------------------------------------------
+class ZenBaseCmd(unrealcmd.Cmd):
+    cloudauthservice = unrealcmd.Opt("", "Name of the service to authorize with when importing from a cloud source")
+
+    def _build(self, target, variant, platform):
+        if target.get_type() == unreal.TargetType.EDITOR:
+            build_cmd = ("_build", "editor", variant)
+        else:
+            platform = platform or unreal.Platform.get_host()
+            build_cmd = ("_build", "target", target.get_name(), platform, variant)
+
+        import subprocess
+        ret = subprocess.run(build_cmd)
+        if ret.returncode:
+            raise SystemExit(ret.returncode)
+
+    def get_binary_path(self, target, variant, platform, build):
+        ue_context = self.get_unreal_context()
+
+        if isinstance(target, str):
+            target = ue_context.get_target_by_name(target)
+        else:
+            target = ue_context.get_target_by_type(target)
+
+        if build:
+            self._build(target, variant, platform)
+
+        variant = unreal.Variant.parse(variant)
+        if build := target.get_build(variant, platform):
+            if build := build.get_binary_path():
+                return build
+
+        raise EnvironmentError(f"No {variant} build found for target '{target.get_name()}'")
+
+    def run_ue_program_target(self, target, variant, build, cmdargs, runargs):
+        binary_path = self.get_binary_path(target, variant, None, build)
+
+        final_args = []
+        if cmdargs: final_args += cmdargs
+        if runargs: final_args += (*unreal.cmdline.read_ueified(*runargs),)
+
+        launch_kwargs = {}
+        if not sys.stdout.isatty():
+            launch_kwargs = { "silent" : True }
+        exec_context = self.get_exec_context()
+        runnable = exec_context.create_runnable(str(binary_path), *final_args);
+        runnable.launch(**launch_kwargs)
+        return True if runnable.is_gui() else runnable.wait()
+
+    def launch_zenserver(self):
+        args = ["-SponsorProcessID=" + str(os.getppid())]
+        self.run_ue_program_target("ZenLaunch", "development", True, args, None)
+
+    def launch_zenserver_if_not_running(self):
+        if not self.is_zenserver_running():
+            self.launch_zenserver()
+
+    def is_zenserver_running(self, hostname="localhost", port=8558):
+        headers = {"Accept":"application/json"}
+        conn = http.client.HTTPConnection(hostname, port=port, timeout=0.1)
+        try:
+            conn.request("GET", "/health/ready", headers=headers)
+            response = conn.getresponse()
+            return response.read().decode().lower() == "ok!"
+        except OSError:
+            return False
+
+        return False
+
+    def find_project_by_id(self, projectid, hostname="localhost", port=8558):
+        self.launch_zenserver_if_not_running()
+        headers = {"Accept":"application/json"}
+        conn = http.client.HTTPConnection(hostname, port=port, timeout=0.1)
+        projects = []
+        try:
+            conn.request("GET", "/prj", headers=headers)
+            response = conn.getresponse()
+            projects = json.loads(response.read().decode())
+        except OSError:
+            return None
+
+        for project in projects:
+            if project['Id'] == projectid:
+                return project
+        return None
+
+    def find_oplog_by_id(self, projectid, oplogid, hostname="localhost", port=8558):
+        self.launch_zenserver_if_not_running()
+        headers = {"Accept":"application/json"}
+        conn = http.client.HTTPConnection(hostname, port=port, timeout=0.1)
+        oplogs = []
+        try:
+            conn.request("GET", f"/prj/{projectid}/oplog/{oplogid}", headers=headers)
+            response = conn.getresponse()
+            oplog = json.loads(response.read().decode())
+        except:
+            return None
+
+        return oplog
+
+    def perform_project_completion(self, prefix, hostname="localhost", port=8558):
+        headers = {"Accept":"application/json"}
+        conn = http.client.HTTPConnection(hostname, port=port, timeout=0.1)
+        projects = []
+        try:
+            conn.request("GET", "/prj", headers=headers)
+            response = conn.getresponse()
+            projects = json.loads(response.read().decode())
+        except OSError:
+            self.launch_zenserver()
+            return
+
+        for project in projects:
+            yield project['Id']
+
+    def perform_oplog_completion(self, prefix, project_filter, hostname="localhost", port=8558):
+        headers = {"Accept":"application/json"}
+        conn = http.client.HTTPConnection(hostname, port=port, timeout=0.1)
+        projects = []
+        try:
+            if project_filter:
+                conn.request("GET", f"/prj/{project_filter}", headers=headers)
+                response = conn.getresponse()
+                projects.append(json.loads(response.read().decode()))
+            else:
+                conn.request("GET", "/prj", headers=headers)
+                response = conn.getresponse()
+                projects = json.loads(response.read().decode())
+        except OSError:
+            self.launch_zenserver()
+            return
+
+        for project in projects:
+            for oplog in project['oplogs']:
+                yield oplog['id']
+
+    def get_project_id_from_project(self, project):
+        if not project:
+            self.print_error("An active project is required for this operation")
+
+        normalized_project_path = str(project.get_path()).replace("\\","/")
+        project_id = project.get_name() + "." + hashlib.md5(normalized_project_path.encode('utf-8')).hexdigest()[:8]
+        return project_id;
+
+
+    def get_project_or_default(self, explicit_project):
+        if explicit_project:
+            return explicit_project
+
+        ue_context = self.get_unreal_context()
+        return self.get_project_id_from_project(ue_context.get_project())
+
+    def _lookup_service_name(self):
+        if getattr(self, "_service_name", None):
+            return
+
+        config = self.get_unreal_context().get_config()
+        default = config.get("Engine", "StorageServers", "Default")
+        if value := default.OAuthProviderIdentifier:
+            self._service_name = value
+        else:
+            cloud = config.get("Engine", "StorageServers", "Cloud")
+            if value := cloud.OAuthProviderIdentifier:
+                self._service_name = value
+
+    def get_exec_context(self):
+        context = super().get_exec_context()
+        if hasattr(self, '_AccessToken') and self._AccessToken:
+            if sys.platform == 'win32':
+                context.get_env()["UE-CloudDataCacheAccessToken"] = self._AccessToken
+            else:
+                context.get_env()["UE_CloudDataCacheAccessToken"] = self._AccessToken
+        return context
+
+    def refresh_zen_token(self):
+        ue_context = self.get_unreal_context()
+        engine = ue_context.get_engine()
+
+        bin_type = "win-x64"
+        if sys.platform == "darwin": bin_type = "osx-x64"
+        elif sys.platform == "linux": bin_type = "linux-x64"
+
+        bin_ext = ".exe"
+        if sys.platform == "darwin": bin_ext = ""
+        elif sys.platform == "linux": bin_ext = ""
+
+        bin_path = engine.get_dir()
+        bin_path /= f"Binaries/DotNET/OidcToken/{bin_type}/OidcToken{bin_ext}"
+        if not bin_path.is_file():
+            raise FileNotFoundError(f"Unable to find '{bin_path}'")
+
+        self._service_name = self.args.cloudauthservice
+        self._lookup_service_name()
+        if not self._service_name:
+            raise ValueError("Unable to discover service name")
+
+        token_dir = tempfile.TemporaryDirectory()
+        token_file_path = f"{token_dir.name}/oidctoken.json"
+        oidcargs = (
+            "--ResultToConsole=true",
+            "--Service=" + str(self._service_name),
+            f"--OutFile={token_file_path}"
+        )
+
+        if project := ue_context.get_project():
+            oidcargs = (*oidcargs, "--Project=" + str(project.get_path()))
+
+        cmd = self.get_exec_context().create_runnable(bin_path, *oidcargs)
+        ret, output = cmd.run2()
+        if ret != 0:
+            return False
+
+        tokenresponse = None
+        with open(token_file_path, "r") as token_file:
+            tokenresponse = json.load(token_file)
+
+        self._AccessToken = tokenresponse['Token']
+        self.print_info("Cloud access token obtained for import operation.")
+
+        return True
+
+    def conform_to_cook_forms(self, oplogid):
+        ue_context = self.get_unreal_context()
+        platforms = ue_context.get_platform_provider()
+        for platform_name in platforms.read_platform_names():
+            platform = platforms.get_platform(platform_name)
+            for form_name in ("client", "game", "server"):
+                try:
+                    cook_form = platform.get_cook_form(form_name)
+                    if cook_form.lower() == oplogid.lower():
+                        return cook_form
+                except:
+                    pass
+        return oplogid
+
+    def perform_target_creation(self, snapshot, projectid = None, oplogid = None):
+        target_projectid = self.get_project_or_default(projectid)
+        if self.find_project_by_id(target_projectid) is None:
+            ue_context = self.get_unreal_context()
+            args = [
+            "project-create",
+            target_projectid,
+            ue_context.get_branch().get_dir(),
+            ue_context.get_engine().get_dir(),
+            ue_context.get_project().get_dir(),
+            ue_context.get_project().get_path()
+            ]
+            self.run_zen_utility(args)
+
+        target_oplogid = oplogid or snapshot['targetplatform']
+
+        ue_context = self.get_unreal_context()
+        target_gcmarker = ue_context.get_project().get_dir() / f"Saved/Cooked/{self.conform_to_cook_forms(target_oplogid)}/ue.projectstore"
+
+        target_gcmarker_dir = os.path.dirname(target_gcmarker)
+        if not os.path.exists(target_gcmarker_dir):
+            os.makedirs(target_gcmarker_dir)
+        remotehostnames = _get_ips(snapshot['targetplatform'])
+        remotehostnames.append("hostname://" + socket.getfqdn())
+        projectstore = {
+            "zenserver" : {
+                "islocalhost" : True,
+                "hostname" : "[::1]",
+                "remotehostnames" : remotehostnames,
+                "hostport" : 8558,
+                "projectid" : target_projectid,
+                "oplogid" : target_oplogid
+            }
+        }
+        with open(target_gcmarker, "w") as target_gcmarker_outfile:
+            json.dump(projectstore, target_gcmarker_outfile, indent='\t')
+
+        args = [
+        "oplog-create",
+        target_projectid,
+        target_oplogid,
+        target_gcmarker,
+        "--force-update"
+        ]
+        self.run_zen_utility(args)
+        return True
+
+    def add_args_from_descriptor(self, snapshot, sourcehost, args):
+        snapshot_type = snapshot['type']
+
+        if snapshot_type in ['cloud', 'zen', 'builds'] and not sourcehost:
+            sourcehost = snapshot['host']
+
+        if snapshot_type == 'cloud':
+            args.append('--cloud')
+            args.append(sourcehost)
+            args.append('--namespace')
+            args.append(snapshot['namespace'])
+            args.append('--bucket')
+            args.append(snapshot['bucket'])
+            args.append('--key')
+            args.append(snapshot['key'])
+        elif snapshot_type == 'builds':
+            args.append('--builds')
+            args.append(sourcehost)
+            args.append('--namespace')
+            args.append(snapshot['namespace'])
+            args.append('--bucket')
+            args.append(snapshot['bucket'])
+            args.append('--builds-id')
+            args.append(snapshot['builds-id'])
+        elif snapshot_type == 'zen':
+            args.append('--zen')
+            args.append(sourcehost)
+            args.append('--source-project')
+            args.append(snapshot['projectid'])
+            args.append('--source-oplog')
+            args.append(snapshot['oplogid'])
+        elif snapshot_type == 'file':
+            args.append('--file')
+            args.append(snapshot['directory'])
+            args.append('--name')
+            args.append(snapshot['filename'])
+        else:
+            self.print_error(f"Unsupported snapshot type {snapshot_type}")
+            return False
+
+        return True
+
+    def perform_import(self, snapshot, sourcehost, projectid, oplogid, asyncimport, forceimport):
+        args = [
+        "oplog-import",
+        self.get_project_or_default(projectid),
+        oplogid or snapshot['targetplatform'],
+        "--ignore-missing-attachments",
+        "--clean"
+        ]
+
+        if asyncimport:
+            args.append("--async")
+
+        if forceimport:
+            args.append("--force")
+
+        if not self.add_args_from_descriptor(snapshot, sourcehost, args):
+            return False
+
+        if snapshot['type'] in ['cloud', 'builds']:
+            if not self.refresh_zen_token():
+                return False
+
+        return self.run_zen_utility(args)
+
+    def perform_clean_post_import(self, snapshot, oplogid):
+        target_oplogid = oplogid or snapshot['targetplatform']
+
+        ue_context = self.get_unreal_context()
+        target_dir = ue_context.get_project().get_dir() / f"Saved/Cooked/{self.conform_to_cook_forms(target_oplogid)}"
+
+        for fsentry in os.scandir(target_dir):
+            if fsentry.is_dir(follow_symlinks=False):
+                shutil.rmtree(fsentry.path)
+            else:
+                if fsentry.name != 'ue.projectstore':
+                    os.unlink(fsentry.path)
+
+    def import_snapshot(self, snapshotdescriptor, sourcehost, projectid, oplogid, asyncimport, forceimport):
+        if not self.perform_target_creation(snapshotdescriptor, projectid, oplogid):
+            self.print_error(f"Error creating import target location")
+            return False
+
+        import_result = self.perform_import(snapshotdescriptor, sourcehost, projectid, oplogid, asyncimport, forceimport)
+        if import_result == 0:
+            self.perform_clean_post_import(snapshotdescriptor, oplogid)
+        return import_result
+
+    def get_workspaces(self):
+        conn = http.client.HTTPConnection(host="localhost", port=8558, timeout=0.1)
+        headers = {"Accept":"application/json"}
+        workspaces_obj = {}
+        try:
+            conn.request('GET', '/ws', headers=headers)
+            response = conn.getresponse()
+            workspaces_obj = json.loads(response.read().decode())
+        except OSError:
+            return None
+
+        if workspaces_obj:
+            return workspaces_obj["workspaces"]
+        return None
+
+    def get_workspace(self, workspaces, base_dir):
+        for workspace in workspaces:
+            if workspace['root_path'] == base_dir:
+                return workspace
+        return None
+
+    def get_workspaceshare(self, workspace, share_dir):
+        for share in workspace.get('shares', []):
+            if share['share_path'] == share_dir:
+                return share
+        return None
+
+    def create_workspace(self, workspacepath, dynamic):
+        workspaces = self.get_workspaces()
+        workspace_exists = False
+        for workspace in workspaces:
+            if workspace['root_path'] == workspacepath:
+                workspace_exists = True
+                break
+
+        if not workspace_exists:
+            args = [
+                "workspace",
+                "create",
+                workspacepath,
+                "--allow-share-create-from-http" if dynamic else None
+            ]
+            if self.run_zen_utility(args) != 0:
+                return None
+
+        workspaces = self.get_workspaces()
+        for workspace in workspaces:
+            if workspace['root_path'] == workspacepath:
+                return workspace
+
+        return None
+
+    def create_workspaceshare(self, workspacepath, sharepath, dynamic):
+        existing_workspace = None
+        workspaces = self.get_workspaces()
+        for workspace in workspaces:
+            if workspace['root_path'] == workspacepath:
+                existing_workspace = workspace
+                break
+
+        if existing_workspace is None:
+            existing_workspace = self.create_workspace(workspacepath, dynamic)
+            if existing_workspace is None:
+                return False
+
+        existing_share = None
+        for share in existing_workspace.get('shares', []):
+            if share["share_path"] == sharepath:
+                existing_share = share
+                break
+
+        if existing_share is None:
+            args = [
+                "workspace-share",
+                "create",
+                existing_workspace["id"],
+                sharepath
+            ]
+            return self.run_zen_utility(args) == 0
+
+        return True
+
+#-------------------------------------------------------------------------------
+class ZenUETargetBaseCmd(ZenBaseCmd):
+    variant = unrealcmd.Arg("development", "Build variant to launch")
+    runargs  = unrealcmd.Arg([str], "Additional arguments to pass to the process being launched")
+    build    = unrealcmd.Opt(True, "Build the target prior to running it")
+
+    def run_ue_program_target(self, target, cmdargs = None):
+        return super().run_ue_program_target(target, self.args.variant, self.args.build, cmdargs, self.args.runargs)
+
+
+#-------------------------------------------------------------------------------
+class ZenUtilityBaseCmd(ZenBaseCmd):
+    zenargs  = unrealcmd.Arg([str], "Additional arguments to pass to zen utility")
+
+    def get_zen_utility_command(self, cmdargs):
+        ue_context = self.get_unreal_context()
+        engine = ue_context.get_engine()
+
+        platform = unreal.Platform.get_host()
+        exe_filename = "zen"
+        if platform == "Win64":
+            exe_filename = exe_filename + ".exe"
+
+        bin_path = engine.get_dir()
+        bin_path /= f"Binaries/{platform}/{exe_filename}"
+        if not bin_path.is_file():
+            raise FileNotFoundError(f"Unable to find '{bin_path}'")
+
+        final_args = []
+        if cmdargs: final_args += cmdargs
+        if self.args.zenargs: final_args += self.args.zenargs
+
+        return bin_path, final_args
+
+
+    def run_zen_utility(self, cmdargs):
+        bin_path, final_args = self.get_zen_utility_command(cmdargs)
+
+        cmd = self.get_exec_context().create_runnable(bin_path, *final_args)
+        return cmd.run()
\ No newline at end of file
diff --git a/Engine/Extras/ushell/channels/unreal/core/pylib/zen/snapshot.py b/Engine/Extras/ushell/channels/unreal/core/pylib/zen/snapshot.py
new file mode 100644
--- /dev/null
+++ b/Engine/Extras/ushell/channels/unreal/core/pylib/zen/snapshot.py
@@ -0,0 +1,518 @@
+# Copyright Epic Games, Inc. All Rights Reserved.
+
+import bisect
+import flow.cmd
+import json
+import os
+import pathlib
+import unreal
+import zen.cmd
+import unrealcmd
+import urllib.request
+
+from abc import ABC, abstractmethod
+
+#-------------------------------------------------------------------------------
+class BuildIndexException(Exception):
+    pass
+
+#-------------------------------------------------------------------------------
+class BuildIndex(ABC):
+    @abstractmethod
+    def get_available_changelists(self, project, platform, branch, runtime, flavor):
+        pass
+
+    @abstractmethod
+    def find_build(self, project, platform, branch, runtime, flavor, buildtype, changelist = 0):
+        pass
+
+    @abstractmethod
+    def get_snapshot_descriptor(self, project, platform, branch, runtime, flavor, buildtype, changelist):
+        pass
+
+#-------------------------------------------------------------------------------
+class UnrealCloudDDCBuildIndex(BuildIndex):
+    def __init__(self, host, access_token, context):
+        self.host = host
+        self.access_token = access_token
+        self.context = context
+
+    # Transform the escaped branch into a string that's valid for use in a cloud namespace
+    def sanitize_branch(self, branch):
+        return branch.lower().replace("++", "").replace("+", "-").replace(".", "-")
+
+    def get_namespace(self, branch):
+        config = self.context.get_config().get("Engine", "StorageServers", "Cloud")
+        if config.BuildsNamespace:
+            return str(config.BuildsNamespace)
+
+        raise BuildIndexException("Unable to determine builds namespace from [StorageServers] config section")
+
+    def get_bucket(self, project, branch, platform, runtime, flavor):
+        branch = self.sanitize_branch(branch)
+        if platform.lower() == "win64":
+            platform = "windows"
+        elif platform.lower() == "android":
+            if not flavor:
+                flavor = "etc2"
+
+        if flavor:
+            platform = platform + "_" + flavor
+
+        platform = platform.lower()
+
+        if runtime.lower() != "game":
+            platform += runtime.lower()
+
+        if platform.lower() == "windowsserver":
+            platform = "linuxserver"
+
+        return f"{project}.oplog.{branch}.{platform}".lower()
+
+    def query_builds(self, project, platform, branch, runtime, flavor, buildtype = "", changelist = 0, try_alternate = True):
+        namespace = self.get_namespace(branch)
+        bucket = self.get_bucket(project, branch, platform, runtime, flavor)
+
+        url = f"{self.host}/api/v2/builds/{namespace}/{bucket}/search"
+        headers = {
+            "Authorization" : f"Bearer {self.access_token}",
+            "Content-Type" : "application/json"
+        }
+
+        query = {}
+        if changelist != 0:
+            query["changelist"] = {"$lte" : changelist}
+
+        query["runtime"] = {"$eq" : runtime}
+        query["branch"] = {"$eq" : branch.replace("+", "/")}
+
+        data = json.dumps({"query": query}).encode()
+        request = urllib.request.Request(url, data=data, headers=headers)
+        response = urllib.request.urlopen(request)
+        body = response.read().decode()
+
+        results = json.loads(body)["results"]
+        results.sort(key=lambda r: r["metadata"]["changelist"])
+
+        if not results and try_alternate:
+            alternate_runtime = ""
+            if runtime == "client" or runtime == "server":
+                alternate_runtime = "game"
+            elif runtime == "game":
+                alternate_runtime = "client"
+
+            results = self.query_builds(project, platform, branch, alternate_runtime, flavor, buildtype, changelist, False)
+            if results:
+                error = f"No builds found for project '{project}' and runtime '{runtime}'.\n"
+                error += f"Builds are available for the '{alternate_runtime}' runtime. Did you mean to use the '{alternate_runtime}' runtime instead of the '{runtime}' runtime?"
+                raise BuildIndexException(error)
+
+        if buildtype:
+            results = [r for r in results if r["metadata"]["hordeTemplateId"].startswith(buildtype)]
+
+        return results
+
+    def get_available_changelists(self, project, platform, branch, runtime, flavor):
+        builds = self.query_builds(project, platform, branch, runtime, flavor)
+        return list(map(lambda build : int(build["metadata"]["changelist"]), builds))
+
+    def find_build(self, project, platform, branch, runtime, flavor, buildtype, changelist):
+        builds = self.query_builds(project, platform, branch, runtime, flavor, buildtype, changelist)
+
+        if len(builds) == 0:
+            return changelist, []
+
+        changelist = int(builds[-1]["metadata"]["changelist"])
+        return changelist, builds
+
+    def get_snapshot_descriptor(self, project, platform, branch, runtime, flavor, buildtype, changelist):
+        builds = self.query_builds(project, platform, branch, runtime, flavor, buildtype, changelist)
+        if len(builds) == 0:
+            return None
+
+        build = builds[-1]
+        builds_id = build["buildId"]
+
+        return int(build["metadata"]["changelist"]), {
+            "type" : "builds",
+            "host" : self.host,
+            "namespace" : self.get_namespace(branch),
+            "bucket" : self.get_bucket(project, branch, platform, runtime, flavor),
+            "builds-id" : builds_id,
+            "targetplatform" : build["metadata"]["cookPlatform"],
+            "name" : build["metadata"]["name"]
+        }
+
+#-------------------------------------------------------------------------------
+class FileshareBuildIndex(BuildIndex):
+    def __init__(self, buildroot, context, host_platform):
+        self.buildroot = buildroot
+        self.context = context
+        self.host_platform = host_platform
+
+    @staticmethod
+    def get_platform_subdir(platform, runtime):
+        runtime_suffix = ""
+        if runtime.lower() != "game":
+            runtime_suffix = runtime.title()
+
+        platform_subdir = None
+
+        if platform.lower() == "win64":
+            platform_subdir = "Windows" + runtime_suffix
+        elif platform.lower() == "android":
+            platform_subdir = "Android"
+        else:
+            platform_subdir = platform.title() + runtime_suffix
+
+        return platform_subdir
+
+    def get_build_root(self):
+        if self.buildroot == "":
+            engine = self.context.get_engine()
+            version = engine.get_info()
+            if not version:
+                raise BuildIndexException("Unable to determine engine version information")
+
+            host_platform_name = self.host_platform
+            root = self.context.get_config().get("Editor", "FileshareBuildIndex", "RootDirectory")
+            path = self.context.get_config().get("Editor", "FileshareBuildIndex", "Path")
+
+            if not root or not path:
+                error = "Unable to determine fileshare build root from editor configuration.\n"
+                error += "To use a fileshare build index, you must have a [FileshareBuildIndex] section in an editor .ini file with the following elements:\n"
+                error += "  RootDirectory(Windows=<fileshare root directory>, Mac=<fileshare root directory>, Linux=<fileshare root directory>)\n"
+                error += "  Path=<shared path>\n"
+                error += "The placeholders {Branch} and {Project} can be used to fill in snapshot information for a given Perforce branch and/or project.\n"
+                error += "For example:\n"
+                error += "  RootDirectory=(Windows=\\\\my.unc.path\\, Mac=/Volumes, Linux=/mnt)\n"
+                error += "  Path=Snapshots/{Branch}/{Project}"
+                raise BuildIndexException(error)
+
+            if host_platform_name == "Win64":
+                self.buildroot = os.path.join(str(root.Windows), str(path))
+            elif host_platform_name == "Mac":
+                self.buildroot = os.path.join(str(root.Mac), str(path))
+            else:
+                self.buildroot = os.path.join(str(root.Linux), str(path))
+
+        return self.buildroot
+
+    def get_builds_dir(self, project, platform, branch, runtime):
+        project_name = self.context.get_project().get_name()
+        project_branch_dir = self.get_build_root().replace("{Branch}", branch).replace("{Project}", project)
+
+        platform_subdir = FileshareBuildIndex.get_platform_subdir(platform, runtime)
+        builds_dir = os.path.join(project_branch_dir, platform_subdir)
+
+        if not os.path.exists(builds_dir):
+            error = f"Unable to reach build index. Please ensure that the build index location is correct and reachable (may require mounting and/or VPN access): {builds_dir}"
+            alternate_runtime = ""
+            if runtime == "client" or runtime == "server":
+                alternate_runtime = "game"
+            elif runtime == "game":
+                alternate_runtime = "client"
+
+            alternate_builds_dir = os.path.join(project_branch_dir, FileshareBuildIndex.get_platform_subdir(platform, alternate_runtime))
+            if os.path.exists(alternate_builds_dir):
+                error += f"\nA build index was available for the '{alternate_runtime}' runtime. Did you mean to use the '{alternate_runtime}' runtime instead of the '{runtime}' runtime?"
+
+            raise BuildIndexException(error)
+
+        return builds_dir
+
+    def get_available_changelists(self, project, platform, branch, runtime, flavor):
+        builds_dir = self.get_builds_dir(project, platform, branch, runtime)
+        if builds_dir == None:
+            return []
+
+        return [int(dir.name) for dir in os.scandir(builds_dir) if dir.name.isdecimal()]
+
+    def get_filtered_builds(self, project, platform, branch, runtime, buildtype, changelist):
+        builds_dir_for_changelist = os.path.join(self.get_builds_dir(project, platform, branch, runtime), str(changelist))
+        builds = []
+
+        try:
+            builds_unfiltered = os.scandir(builds_dir_for_changelist)
+            if buildtype:
+                for build in builds_unfiltered:
+                    build_name_lower = build.name.lower()
+                    if build_name_lower.startswith(buildtype.lower()):
+                        builds.append(build.name)
+            else:
+                for build in builds_unfiltered:
+                    builds.append(build.name)
+        except:
+            pass
+
+        return builds
+
+    def find_build(self, project, platform, branch, runtime, flavor, buildtype, changelist):
+        builds_dir_for_changelist = ""
+        builds = []
+
+        current_iteration = 0
+        max_iterations = 20
+
+        buildchangelists = None
+
+        while not builds and (current_iteration < max_iterations):
+            builds = self.get_filtered_builds(project, platform, branch, runtime, buildtype, changelist)
+
+            if not builds:
+                if buildchangelists is None:
+                    try:
+                        buildchangelistdirs = os.scandir(self.get_builds_dir(project, platform, branch, runtime))
+                    except:
+                        buildchangelistdirs = []
+
+                    buildchangelists = []
+                    for buildchangelistdir in buildchangelistdirs:
+                        try:
+                            buildchangelists.append(int(buildchangelistdir.name))
+                        except:
+                            pass
+                    buildchangelists.sort()
+                i = bisect.bisect_right(buildchangelists, changelist-1)
+                if i != 0:
+                    changelist = buildchangelists[i-1]
+            current_iteration = current_iteration + 1
+
+        return int(changelist), builds
+
+    def get_snapshot_descriptor(self, project, platform, branch, runtime, flavor, buildtype, changelist):
+        builds_dir = self.get_builds_dir(project, platform, branch, runtime)
+        if builds_dir is None:
+            return None
+
+        found_changelist, builds = self.find_build(project, platform, branch, runtime, flavor, buildtype, changelist)
+        if not builds:
+            raise BuildIndexException("No build found")
+
+        descriptorfile = os.path.join(builds_dir, str(found_changelist), builds[0])
+        try:
+            with (open(descriptorfile, "rt")) as file:
+                descriptors = json.load(file)
+                return int(found_changelist), descriptors["snapshots"][0]
+        except:
+            raise BuildIndexException(f"Error loading snapshot descriptor from file {descriptorfile}")
+
+#-------------------------------------------------------------------------------
+class _Base(zen.cmd.ZenUtilityBaseCmd):
+    runtime     = flow.cmd.Arg(str, "Type of runtime to use (client, server, game)")
+    platform    = flow.cmd.Arg(str, "Platform to use")
+    changelist  = flow.cmd.Arg(-1, "Changelist to query snapshot for (defaults to current changelist)")
+    buildroot   = unrealcmd.Opt("", "The root directory to use for builds (defaults to fileshare roots specified in [FileshareBuildIndex] editor configuration section")
+    nofileshare = flow.cmd.Opt(False, "Find snapshots stored in Unreal Cloud DDC rather than using the network fileshare")
+    fileshare   = flow.cmd.Opt(False, "Find snapshots stored on the network fileshare instead of searching Unreal Cloud")
+    cloudhost   = unrealcmd.Opt("", "The cloud host to use for builds (defaults to cloud host specified in [StorageServers] engine configuration section)")
+    flavor      = unrealcmd.Opt("", "Flavor of build to retrieve (e.g. 'ASTC' or 'ETC2' for Android builds).")
+    buildtype   = flow.cmd.Opt("", "Specific build type to get like nightly, standard, etc... (defaults to any type available)")
+    build_index = None
+
+    def get_build_index(self):
+        if not self.build_index:
+            if self.args.nofileshare:
+                self.print_info("--nofileshare is now the default for .snapshot commands. If you would like to use the fileshare instead of the Unreal Cloud build index, use --fileshare.")
+            if not self.args.fileshare:
+                if not self.refresh_zen_token():
+                    self.print_error("Unable to get cloud access token")
+                    return None
+
+                if not self.args.cloudhost:
+                    config = self.get_unreal_context().get_config().get("Engine", "StorageServers", "Cloud")
+                    if config.Host:
+                        self.args.cloudhost = str(config.Host).split(";")[0]
+                    else:
+                        error = "Unable to determine cloud host from [StorageServers] config section\n"
+                        error += "To use an Unreal Cloud build index, you must have a [StorageServers] config section in an engine .ini file with a value Cloud=(Host=https://my.cloud.host)"
+                        raise BuildIndexException(error)
+
+                self.build_index = UnrealCloudDDCBuildIndex(self.args.cloudhost, self._AccessToken, self.get_unreal_context())
+                self.print_info(f"Finding snapshots from builds URL {self.build_index.host}")
+            else:
+                self.build_index = FileshareBuildIndex(self.args.buildroot, self.get_unreal_context(), self.get_host_platform().get_name())
+                self.print_info(f"Finding snapshots from fileshare build root {self.build_index.get_build_root()}")
+
+        return self.build_index
+
+    def complete_changelist(self, prefix):
+        context = self.get_unreal_context()
+        engine = context.get_engine()
+        version = engine.get_info()
+        if version:
+            branch_name = version["BranchName"]
+            project_name = context.get_project().get_name()
+
+            changelists = self.get_build_index().get_available_changelists(project_name, self.args.platform, branch_name, self.args.runtime, self.args.flavor)
+            for changelist in changelists:
+                yield changelist
+
+    def complete_runtime(self, prefix):
+        yield "client"
+        yield "server"
+        yield "game"
+
+    def complete_platform(self, prefix):
+        context = self.get_unreal_context()
+        engine = context.get_engine()
+        version = engine.get_info()
+        if version:
+            branch_name = version["BranchName"]
+            for platform in unrealcmd.Cmd.complete_platform(self, prefix):
+                if platform == "android":
+                    if self.args.runtime == "server":
+                        continue
+
+                yield platform
+
+    def get_changelist(self):
+        changelist = self.args.changelist
+        if changelist < 1:
+            changelist = self.get_unreal_context().get_engine().get_info()["Changelist"]
+        return changelist
+
+    def get_project_name(self):
+        return self.get_unreal_context().get_project().get_name()
+
+    def get_branch_name(self):
+        return self.get_unreal_context().get_engine().get_info()["BranchName"]
+
+#-------------------------------------------------------------------------------
+class Find(_Base):
+    """ Finds the best available snapshot by runtime and platform while filling
+    in remaining details like project and branch using ushell context information."""
+
+    buildtype = flow.cmd.Opt("", "Specific build type to get like nightly, standard, etc... (defaults to any type available)")
+
+    def complete_projectid(self, prefix):
+        return self.perform_project_completion(prefix)
+
+    def complete_oplog(self, prefix):
+        return self.perform_oplog_completion(prefix, self.get_project_or_default(self.args.projectid))
+
+    def main(self):
+        found_changelist = self.get_closest_changelist()
+        if found_changelist < 1:
+            return False
+
+        print(found_changelist)
+        changelist = self.get_changelist()
+        if (found_changelist == changelist):
+            print(f"Exact match for target changelist")
+        else:
+            print(f"Closest candidate for target changelist '{changelist}'")
+
+    def get_closest_changelist(self):
+        changelist = self.get_changelist()
+        project = self.get_project_name()
+        branch = self.get_branch_name()
+
+        try:
+            found_changelist, builds = self.get_build_index().find_build(project, self.args.platform, branch, self.args.runtime, self.args.flavor, self.args.buildtype, changelist)
+
+            if not builds:
+                self.print_error(f"No build found")
+                return 0
+
+        except BuildIndexException as e:
+            self.print_error(e)
+            return 0
+
+        return found_changelist
+
+#-------------------------------------------------------------------------------
+class Get(_Base):
+    """ Gets a cooked data snapshot by runtime and platform while filling in
+    remaining details like project, branch, and even changelist using ushell
+    context information.  If a snapshot cannot be found for the precise changelist,
+    offers to pick nearest preceding changelist that matches the other criteria."""
+
+    buildtype   = flow.cmd.Opt("", "Specific build type to get like nightly, standard, etc... (defaults to any type available)")
+    projectid   = unrealcmd.Opt("", "The zen project ID to import into (defaults to an ID based on the current ushell project)")
+    oplog       = unrealcmd.Opt("", "The zen oplog to import into (defaults to the oplog name in the snapshot)")
+    sourcehost  = unrealcmd.Opt("", "The source host to import from (defaults to the host specified in the snapshot)")
+    asyncimport = unrealcmd.Opt(False, "Trigger import but don't wait for completion")
+    forceimport = unrealcmd.Opt(False, "Force import of all attachments")
+
+    def complete_projectid(self, prefix):
+        return self.perform_project_completion(prefix)
+
+    def complete_oplog(self, prefix):
+        return self.perform_oplog_completion(prefix, self.get_project_or_default(self.args.projectid))
+
+    def main(self):
+        self.launch_zenserver_if_not_running()
+        changelist = self.get_changelist()
+        project_name = self.get_project_name()
+        branch = self.get_branch_name()
+
+        try:
+            found_changelist, snapshotdescriptor = self.get_build_index().get_snapshot_descriptor(project_name, self.args.platform, branch, self.args.runtime, self.args.flavor, self.args.buildtype, changelist)
+
+            while found_changelist != changelist:
+                self.print_error(f"No snapshot found for the specified project({project_name}), branch({branch}), platform({self.args.platform}), runtime({self.args.runtime}), and changelist({changelist})")
+                print("Now you have a few courses of action available to you;")
+                print("  use closest [p]receding snapshot changelist")
+                print("  retry [Enter]")
+                print("  abort [Ctrl-C]")
+
+                choice = input("Which one do you fancy? [p/Enter/Ctrl-C] ")
+                if choice == "p":
+                    break
+
+                found_changelist, snapshotdescriptor = self.get_build_index().get_snapshot_descriptor(project_name, self.args.platform, version["BranchName"], self.args.runtime, self.args.flavor, self.args.buildtype, changelist)
+
+            if not snapshotdescriptor:
+                self.print_error(f"No build found")
+                return False
+
+            print(f"Importing snapshot for specified project({project_name}), branch({branch}), platform({self.args.platform}), runtime({self.args.runtime}), and changelist({changelist}) from {snapshotdescriptor['type']} snapshot:")
+            print(snapshotdescriptor)
+
+            self.import_snapshot(snapshotdescriptor, self.args.sourcehost, self.args.projectid, self.args.oplog, self.args.asyncimport, self.args.forceimport)
+            self.post_import_snapshot(found_changelist)
+        except BuildIndexException as e:
+            self.print_error(e)
+
+    def post_import_snapshot(self, snapshot_changelist):
+        return
+
+#-------------------------------------------------------------------------------
+class List(_Base):
+    """ Lists available snapshots by runtime and platform while filling in
+    remaining details like project and branch using ushell context information."""
+
+    def complete_projectid(self, prefix):
+        return self.perform_project_completion(prefix)
+
+    def complete_oplog(self, prefix):
+        return self.perform_oplog_completion(prefix, self.get_project_or_default(self.args.projectid))
+
+    def main(self):
+        changelists = self.get_snapshot_changelists()
+        if not changelists:
+            return False
+
+        current_changelist = self.get_changelist()
+        for changelist in changelists:
+            if int(changelist) == current_changelist:
+                print(f"{changelist}*")
+            else:
+                print(f"{changelist}")
+
+    def get_snapshot_changelists(self):
+        project_name = self.get_project_name()
+        branch_name = self.get_branch_name()
+
+        try:
+            changelists = self.get_build_index().get_available_changelists(project_name, self.args.platform, branch_name, self.args.runtime, self.args.flavor)
+        except BuildIndexException as e:
+            self.print_error(e)
+            return []
+
+        if not changelists:
+            self.print_error("No build found")
+            return []
+
+        return changelists
