Updated ushell [42282047]
#rnx

diff --git a/Engine/Extras/ushell/.version b/Engine/Extras/ushell/.version
--- a/Engine/Extras/ushell/.version
+++ b/Engine/Extras/ushell/.version
@@ -1 +1 @@
-41861897
\ No newline at end of file
+42282047
\ No newline at end of file
diff --git a/Engine/Extras/ushell/channels/flow/core/system/flow/cmd.py b/Engine/Extras/ushell/channels/flow/core/system/flow/cmd.py
--- a/Engine/Extras/ushell/channels/flow/core/system/flow/cmd.py
+++ b/Engine/Extras/ushell/channels/flow/core/system/flow/cmd.py
@@ -252,8 +252,10 @@ class Cmd(_flick.Cmd, _ArgOverrider):
     def get_exec_context(self):
         return _ExecContext()
 
-    @staticmethod
-    def summarise(main_method):
+    @classmethod
+    def summarise(cls, main_method):
+        cls.nosummary = Opt(False, "Do not print the result/time summary tail")
+
         import time
         import datetime
 
@@ -270,9 +272,12 @@ class Cmd(_flick.Cmd, _ArgOverrider):
 
             return result
 
-        def begin(*args, **kwargs):
+        def begin(self, *args, **kwargs):
+            if self.args.nosummary:
+                return main_method(self, *args, **kwargs)
+
             start_time = time.time()
-            result = main_method(*args, **kwargs)
+            result = main_method(self, *args, **kwargs)
             if getattr(result, "cr_code", None):
                 async def inner(result):
                     result = await result
@@ -280,4 +285,5 @@ class Cmd(_flick.Cmd, _ArgOverrider):
                 return inner(result)
 
             return finish(start_time, result)
+
         return begin
diff --git a/Engine/Extras/ushell/channels/unreal/core/cmds/cook.py b/Engine/Extras/ushell/channels/unreal/core/cmds/cook.py
--- a/Engine/Extras/ushell/channels/unreal/core/cmds/cook.py
+++ b/Engine/Extras/ushell/channels/unreal/core/cmds/cook.py
@@ -31,8 +31,10 @@ class _Cook(unrealcmd.MultiPlatformCmd):
     noxge        = unrealcmd.Opt(False, "Disable XGE-based shader compilation")
     unpretty     = unrealcmd.Opt(False, "Turns off colourful pretty-printing")
     attach       = unrealcmd.Opt(False, "Attach a debugger to the cook")
+    build        = unrealcmd.Opt(False, "Build editor binaries prior to cooking")
     debug        = unrealcmd.Opt(False, "Use debug executables")
     nounattended = unrealcmd.Opt(False, "Skip implicitly specified '-unattended'")
+    versioned    = unrealcmd.Opt(False, "Skip implicitly specified '-unversioned'")
 
     @unrealcmd.Cmd.summarise
     def _cook(self, cook_form, exec_context=None):
@@ -44,6 +46,14 @@ class _Cook(unrealcmd.MultiPlatformCmd):
             self.print_error("An active project is required to cook")
             return False
 
+        # Build binaries
+        if self.args.build:
+            import subprocess as sp
+            variant = "debug" if self.args.debug else "development"
+            proc = sp.run(("_build", "editor", variant, "--nosummary"))
+            if proc.returncode:
+                return proc.returncode
+
         # Prepare arguments
         cultures = self.args.cultures.replace(",", "+")
 
@@ -51,7 +61,7 @@ class _Cook(unrealcmd.MultiPlatformCmd):
             ("-targetplatform=" + cook_form) if cook_form else None,
             ("-cookcultures=" + cultures) if cultures else None,
             "-unattended" if not self.args.nounattended else None,
-            "-unversioned",
+            "-unversioned" if not self.versioned else None,
             "-stdout",
             "-cookonthefly" if self.args.onthefly else None,
             "-iterate" if self.args.iterate else None,
diff --git a/Engine/Extras/ushell/channels/unreal/core/pylib/unreal/_context.py b/Engine/Extras/ushell/channels/unreal/core/pylib/unreal/_context.py
--- a/Engine/Extras/ushell/channels/unreal/core/pylib/unreal/_context.py
+++ b/Engine/Extras/ushell/channels/unreal/core/pylib/unreal/_context.py
@@ -381,13 +381,15 @@ class _Engine(_UfsNode):
         }
 
         build_version_path = self._dir / "Build/Build.version"
-        if build_version_path.is_file():
-            with build_version_path.open() as x:
-                try:
-                    import json
-                    ret.update(json.load(x))
-                except:
-                    pass
+        if not build_version_path.is_file():
+            return ret
+
+        with build_version_path.open() as x:
+            try:
+                import json
+                ret.update(json.load(x))
+            except:
+                pass
 
         return ret
 
diff --git a/Engine/Extras/ushell/channels/unreal/perforce/cmds/sync.py b/Engine/Extras/ushell/channels/unreal/perforce/cmds/sync.py
--- a/Engine/Extras/ushell/channels/unreal/perforce/cmds/sync.py
+++ b/Engine/Extras/ushell/channels/unreal/perforce/cmds/sync.py
@@ -25,7 +25,8 @@ class Sync(_SyncBase):
     """ Syncs the current project and engine directories to the given changelist
     (or latest if none is provided). If the '--all' option is specified then the
     branch will be searched locally for existing .uproject files, scheduling each
-    one to be synced.
+    one to be synced. Specifying 'have' as the chagelist to sync will use the
+    same value as returned by '.info'.
 
     The sync can be filtered with a .p4sync.txt file. Lines prefixed with a '-'
     (e.g. "-.../SootySweep/...") will be excluded from the sync and anything
@@ -42,15 +43,16 @@ class Sync(_SyncBase):
         -.../Android/...
         -/Engine/Source/...
         -*.uasset """
-    changelist = unrealcmd.Arg(-1, "The changelist to sync to ('now' if unspecified)")
+    changelist = unrealcmd.Arg("", "The changelist to sync to ('now' if unspecified)")
     noresolve  = unrealcmd.Opt(False, "Do not run 'p4 resolve -am' after syncing")
     dryrun     = unrealcmd.Opt(False, "Only pretend to do the sync")
     all        = unrealcmd.Opt(False, "Sync all the branch's projects found locally")
     addprojs   = unrealcmd.Opt("", "Comma-separated names of additional projects to sync")
     clobber    = unrealcmd.Opt(False, "Clobber writable files when syncing")
-    nosummary  = unrealcmd.Opt(False, "Do not print the result-time summary at the end")
     echo       = unrealcmd.Opt(False, "Echo depot paths as they are synced")
 
+    complete_changelist = ("have", )
+
     def complete_addprojs(self, prefix):
         # Fake a _local_root
         local_root = Path(os.getcwd())
@@ -72,16 +74,10 @@ class Sync(_SyncBase):
         return (x.stem for x in  branch.read_projects())
 
     @unrealcmd.Cmd.summarise
-    def _main_summarised(self):
-        return self._main_impl()
-
     def main(self):
         self._client_spec_restore = None
         try:
-            if self.args.nosummary:
-                return self._main_impl()
-            else:
-                return self._main_summarised()
+            return self._main_impl()
         finally:
             try:
                 if self._client_spec_restore:
@@ -215,11 +211,28 @@ class Sync(_SyncBase):
                 print(uproj_path.stem)
                 syncer.add_path(str(uproj_path.parent) + "/...")
 
+    def _get_have_cl(self):
+        ue_context = self._try_get_unreal_context()
+        if not ue_context or ue_context.get_branch() is None:
+            raise EnvironmentError("No UE branch found - cannot determine 'have' changelist")
+
+        engine = ue_context.get_engine()
+        info = engine.get_info()
+        ret = int(info["Changelist"])
+        if ret <= 0:
+            raise ValueError(f"Unexpected branch changelist '{ret}'")
+
+        return ret
+
     def _main_impl(self):
         self._setup()
 
         # Determine the changelist to sync
         sync_cl = self.args.changelist
+        if sync_cl == "have":
+            sync_cl = self._get_have_cl()
+
+        sync_cl = int(sync_cl or 0)
         if sync_cl < 0:
             sync_cl = int(P4.changes(self._branch_root + "...", m=1).change)
 
