Updated ushell [42427842]

#rnx

diff --git a/Engine/Extras/ushell/.version b/Engine/Extras/ushell/.version
--- a/Engine/Extras/ushell/.version
+++ b/Engine/Extras/ushell/.version
@@ -1 +1 @@
-42282047
\ No newline at end of file
+42427842
\ No newline at end of file
diff --git a/Engine/Extras/ushell/channels/unreal/perforce/cmds/sync.py b/Engine/Extras/ushell/channels/unreal/perforce/cmds/sync.py
--- a/Engine/Extras/ushell/channels/unreal/perforce/cmds/sync.py
+++ b/Engine/Extras/ushell/channels/unreal/perforce/cmds/sync.py
@@ -232,7 +232,7 @@ class Sync(_SyncBase):
         if sync_cl == "have":
             sync_cl = self._get_have_cl()
 
-        sync_cl = int(sync_cl or 0)
+        sync_cl = int(sync_cl or -1)
         if sync_cl < 0:
             sync_cl = int(P4.changes(self._branch_root + "...", m=1).change)
 
