Updated ushell [42565520]
#rnx

diff --git a/Engine/Extras/ushell/.version b/Engine/Extras/ushell/.version
--- a/Engine/Extras/ushell/.version
+++ b/Engine/Extras/ushell/.version
@@ -1 +1 @@
-42427842
\ No newline at end of file
+42565520
\ No newline at end of file
diff --git a/Engine/Extras/ushell/channels/unreal/core/cmds/getbuild.py b/Engine/Extras/ushell/channels/unreal/core/cmds/getbuild.py
--- a/Engine/Extras/ushell/channels/unreal/core/cmds/getbuild.py
+++ b/Engine/Extras/ushell/channels/unreal/core/cmds/getbuild.py
@@ -400,11 +400,11 @@ If you find multiple builds that match you will be given the option to pick one,
                 "--host=" + str(host) if host != None else "--override-host=" + str(override_host),
                 "--namespace=" + str(namespace),
                 "--bucket=" + bucket,
-                "--oidctoken-exe-path=" + self.get_oidc_token_path(),
+                "--oidctoken-exe-path=\"" + self.get_oidc_token_path() + "\"",
                 "--clean" if self.args.clean else "",
                 "--verbose" if self.args.verbose else "",
-                build_dest,
-                build_id
+                "--local-path=\"" + build_dest + "\"",
+                "--build-id=" + build_id
             )
 
             result = self.run_zen(args)
diff --git a/Engine/Extras/ushell/channels/unreal/perforce/cmds/cherrypick.py b/Engine/Extras/ushell/channels/unreal/perforce/cmds/cherrypick.py
--- a/Engine/Extras/ushell/channels/unreal/perforce/cmds/cherrypick.py
+++ b/Engine/Extras/ushell/channels/unreal/perforce/cmds/cherrypick.py
@@ -33,6 +33,7 @@ class Cherrypick(flow.cmd.Cmd):
     changelist    = flow.cmd.Arg([int], "Changelist or shelve to edit-integrate")
     path          = flow.cmd.Opt("", "Restrict cherrypick to the given path (e.g. Engine/...)")
     saferesolve   = flow.cmd.Opt(False, "Resolve safely and not automatically")
+    noresolve     = flow.cmd.Opt(False, "Skip the resolve step")
     dryrun        = flow.cmd.Opt(False, "Only pretend to do the cherrypick")
     force         = flow.cmd.Opt(False, "Force the operation through")
     novalidate    = flow.cmd.Opt(False, "Don't run the validation step")
@@ -294,25 +295,28 @@ class Cherrypick(flow.cmd.Cmd):
             else:
                 self.print_error(msg)
 
-        resolve_args = {
-            "as" : self.args.saferesolve,
-            "am" : not self.args.saferesolve,
-            "f"  : self.args.force,
-            "c"  : dest_cl,
-        }
-        resolve = P4.resolve(**resolve_args)
-        for item in resolve.read(on_error=print_error):
-            if getattr(item, "clientFile", None):
-                resolve_count += 1
-                print("\r" + str(resolve_count), "file(s) resolved", end="")
-        if resolve_count:
-            print("")
+        if self.args.noresolve:
+            print("Skipping resolve")
+        else:
+            resolve_args = {
+                "as" : self.args.saferesolve,
+                "am" : not self.args.saferesolve,
+                "f"  : self.args.force,
+                "c"  : dest_cl,
+            }
+            resolve = P4.resolve(**resolve_args)
+            for item in resolve.read(on_error=print_error):
+                if getattr(item, "clientFile", None):
+                    resolve_count += 1
+                    print("\r" + str(resolve_count), "file(s) resolved", end="")
+            if resolve_count:
+                print("")
 
-        # Report conflicted files
-        resolve = P4.resolve(c=dest_cl, n=True)
-        for item in resolve.read(on_error=False):
-            if name := getattr(item, "fromFile", None):
-                self.print_error(name[len(src_root):])
+            # Report conflicted files
+            resolve = P4.resolve(c=dest_cl, n=True)
+            for item in resolve.read(on_error=False):
+                if name := getattr(item, "fromFile", None):
+                    self.print_error(name[len(src_root):])
 
         # Queue up some details to display when everything's complete
         class OnExit(object):
