PR #11879: When ushell dispatches two separate p4 processes a deadlock can occur?

#rnx

diff --git a/Engine/Extras/ushell/channels/unreal/perforce/pylib/peafour.py b/Engine/Extras/ushell/channels/unreal/perforce/pylib/peafour.py
--- a/Engine/Extras/ushell/channels/unreal/perforce/pylib/peafour.py
+++ b/Engine/Extras/ushell/channels/unreal/perforce/pylib/peafour.py
@@ -1,9 +1,11 @@
 # Copyright Epic Games, Inc. All Rights Reserved.
 
+import errno
 import types
 import marshal
 import threading
 import subprocess as sp
+from queue import SimpleQueue
 from pathlib import PurePath
 
 #-------------------------------------------------------------------------------
@@ -104,27 +106,62 @@ class _P4Command(object):
         proc = sp.Popen(self._command, stdout=sp.PIPE, stdin=stdin)
         self._proc = proc
 
-        if stdin:
-            def stdin_thread_entry():
+        def _stdin_writethread():
+            try:
+                if input_data:
+                    marshal.dump(input_data, proc.stdin, 0)
+                for args in self._stdin_args:
+                    for arg in args:
+                        arg = str(arg).encode() + b"\n"
+                        proc.stdin.write(arg)
+            except BrokenPipeError:
+                pass
+            except OSError as exc:
+                if exc.errno == errno.EINVAL:
+                    pass
+                else:
+                    raise
+            finally:
                 try:
-                    if input_data:
-                        marshal.dump(input_data, proc.stdin, 0)
-                    for args in self._stdin_args:
-                        for arg in args:
-                            arg = str(arg).encode() + b"\n"
-                            proc.stdin.write(arg)
-                except (BrokenPipeError, OSError):
+                    proc.stdin.close()
+                except BrokenPipeError:
                     pass
-                finally:
-                    try: proc.stdin.close()
-                    except: pass
+                except OSError as exc:
+                    if exc.errno == errno.EINVAL:
+                        pass
+                    else:
+                        raise
 
-            stdin_thread = threading.Thread(target=stdin_thread_entry)
+        if stdin:
+            stdin_thread = threading.Thread(target=_stdin_writethread)
+            stdin_thread.daemon = True
             stdin_thread.start()
 
+        # sentinel to signal end of queue
+        _EOF = object()
+        q = SimpleQueue()
+
+        def _stdout_readthread():
+            try:
+                while proc.poll() is None:
+                    try:
+                        result = marshal.load(proc.stdout)
+                        q.put(result)
+                    except (EOFError, ValueError):
+                        # the sub-process can close the stdout file before the
+                        # process has terminated, which will raise a ValueError
+                        break
+            finally:
+                q.put(_EOF)
+
+        stdout_thread = threading.Thread(target=_stdout_readthread)
+        stdout_thread.daemon = True
+        stdout_thread.start()
+
         while True:
-            try: result = marshal.load(proc.stdout)
-            except EOFError: break
+            result = q.get()
+            if result is _EOF:
+                break
 
             code = result[b"code"]
             del result[b"code"]
@@ -165,6 +202,8 @@ class _P4Command(object):
 
         if stdin:
             stdin_thread.join()
+        stdout_thread.join()
+
         self._close_proc()
 
 #-------------------------------------------------------------------------------
