Updated ushell [42762466]
#rnx

diff --git a/Engine/Extras/ushell/.version b/Engine/Extras/ushell/.version
--- a/Engine/Extras/ushell/.version
+++ b/Engine/Extras/ushell/.version
@@ -1 +1 @@
-42565520
\ No newline at end of file
+42762466
\ No newline at end of file
diff --git a/Engine/Extras/ushell/channels/unreal/core/cmds/perftest.py b/Engine/Extras/ushell/channels/unreal/core/cmds/perftest.py
--- a/Engine/Extras/ushell/channels/unreal/core/cmds/perftest.py
+++ b/Engine/Extras/ushell/channels/unreal/core/cmds/perftest.py
@@ -11,6 +11,7 @@ class PerfTestType(Enum):
     REPLAY      = 0x01
     MATERIAL    = 0x02
     CAMERA      = 0x03
+    DEFAULT     = 0x04
 
     def as_arg(self):
         match self:
@@ -22,6 +23,8 @@ class PerfTestType(Enum):
                 return "AutomatedPerfTest.MaterialTest"
             case PerfTestType.CAMERA:
                 return "AutomatedPerfTest.StaticCameraTest"
+            case PerfTestType.DEFAULT:
+                return "AutomatedPerfTest.DefaultTest"
 
 
 #-------------------------------------------------------------------------------
@@ -102,7 +105,7 @@ class _PerfTestConfig(object):
 
         # Return default values if a subtest iteration value
         # does not exist.
-        DEFAULT_ITERATIONS  = 3
+        DEFAULT_ITERATIONS  = 1
         PERF_ITERATIONS     = 6
         LLM_ITERATIONS      = 3
         INSIGHTS_ITERATIONS = 1
@@ -182,7 +185,7 @@ class PerfTestBaseCmd(unrealcmd.Cmd):
             self.print_error("No project found in context")
             return PerfTestResult.ERROR
 
-        test_id = f"{project.get_name()}-{self.args.subtest}-autoperftest-ushell"
+        test_id = f"{project.get_name()}-{subtest.name.lower()}-autoperftest-ushell"
         test_id = test_id if self.args.testid == "" else self.args.testid
 
         if not target:
@@ -292,6 +295,41 @@ class PerfTestBaseCmd(unrealcmd.Cmd):
             dir = project.get_dir()
         return dir.glob(pattern)
 
+    def _find_config(self, config_section:str, config_key:str):
+        config = self.get_unreal_context().get_config()
+        return config.get("Engine", config_section, config_key)
+
+    def _get_available_sequences(self) -> list[str]:
+        config_section  = "/Script/AutomatedPerfTesting.AutomatedSequencePerfTestProjectSettings"
+        config_key      = "MapsAndSequencesToTest"
+        sequences = self._find_config(config_section, config_key)
+        return [sequence.ComboName for sequence in sequences]
+
+    def _get_available_replays(self) -> list[str]:
+        config_section  = "/Script/AutomatedPerfTesting.AutomatedReplayPerfTestProjectSettings"
+        config_key      = "ReplaysToTest"
+        replays = self._find_config(config_section, config_key)
+        return [replay.FilePath for replay in replays]
+
+    def _get_available_static_cameratests(self) -> list[str]:
+        config_section  = "/Script/AutomatedPerfTesting.AutomatedStaticCameraPerfTestProjectSettings"
+        config_key      = "MapsToTest"
+        return self._find_config(config_section, config_key)
+
+    def _get_available_materials(self) -> list[str]:
+        config_section  = "/Script/AutomatedPerfTesting.AutomatedMaterialPerfTestProjectSettings"
+        config_key      = "MaterialsToTest"
+        return self._find_config(config_section, config_key)
+
+    def _make_sequence_args(self, sequence:str) -> tuple[str]:
+        return (f"-AutomatedPerfTest.SequencePerfTest.MapSequenceName={sequence}",)
+
+    def _make_replay_args(self, replay:str) -> tuple[str]:
+        return (f"-AutomatedPerfTest.ReplayPerfTest.ReplayName={replay}",)
+
+    def _make_static_cam_args(self, map:str) -> tuple[str]:
+        return (f"-AutomatedPerfTest.StaticCameraPerfTest.MapName={map}",)
+
     #--- Tab Auto Complete
 
     def complete_subtest(self, prefix):
@@ -308,20 +346,12 @@ class Sequence(PerfTestBaseCmd):
     performance metrics.
     """
     SequenceComboName   = unrealcmd.Arg(str,  "Sequence Combo Name")
-    config_section      = "/Script/AutomatedPerfTesting.AutomatedSequencePerfTestProjectSettings"
-    config_key          = "MapsAndSequencesToTest"
 
     def main(self):
-        args = (f"-AutomatedPerfTest.SequencePerfTest.MapSequenceName={self.args.SequenceComboName}",)
+        args = self._make_sequence_args(self.args.SequenceComboName)
         return self.run_perftest(PerfTestType.SEQUENCE, *args)
 
-    def _get_available_sequences(self) -> list[str]:
-        config = self.get_unreal_context().get_config()
-        sequences = config.get("Engine", self.config_section, self.config_key)
-        return [sequence.ComboName for sequence in sequences]
-
     #--- Tab Auto Complete
-
     def complete_SequenceComboName(self, prefix):
         yield from self._get_available_sequences()
 
@@ -331,23 +361,29 @@ class Replay(PerfTestBaseCmd):
     """
     Replay Test - Runs performance test with given replay file.
     """
+
     ReplayFile  = unrealcmd.Arg(str,  "Replay file")
 
     def _find_replay(self, name):
         path = Path(name)
-        if path.is_file():
+        if path.is_file() and path.exists():
             return path
+
         self.print_info("Searching for replay file...")
         return self._find_file(name, "replay")
 
     def main(self):
         replay = self._find_replay(self.args.ReplayFile)
-        if replay is None:
+        if replay is None or not replay.exists():
             self.print_error(f"Could not find '{self.args.ReplayFile}'")
-            return 1
-        args = (f"-AutomatedPerfTest.ReplayPerfTest.ReplayName={replay}",)
+            return PerfTestResult.ERROR
+        args = self._make_replay_args(replay)
         return self.run_perftest(PerfTestType.REPLAY, *args)
 
+    #--- Tab Auto Complete
+    def complete_ReplayFile(self, prefix):
+        yield from self._get_available_replays()
+
 
 #-------------------------------------------------------------------------------
 class Material(PerfTestBaseCmd):
@@ -366,5 +402,19 @@ class StaticCamera(PerfTestBaseCmd):
     """
     MapName  = unrealcmd.Arg(str,  "Static Camera Test Map Name")
     def main(self):
-        args = (f"-AutomatedPerfTest.StaticCameraPerfTest.MapName={self.args.MapName}",)
-        return self.run_perftest(PerfTestType.CAMERA, *args)
\ No newline at end of file
+        args = self._make_static_cam_args(self.args.MapName)
+        return self.run_perftest(PerfTestType.CAMERA, *args)
+
+    def complete_MapName(self, prefix):
+        yield from self._get_available_static_cameratests()
+
+
+#-------------------------------------------------------------------------------
+class PerfTestDefault(PerfTestBaseCmd):
+    """
+    Perf Test Default - Finds the default (or first) perf test configured for
+    the project and executes it.
+    """
+    def main(self):
+        args = (None,)
+        return self.run_perftest(PerfTestType.DEFAULT, *args)
\ No newline at end of file
diff --git a/Engine/Extras/ushell/channels/unreal/core/debuggers/rider.py b/Engine/Extras/ushell/channels/unreal/core/debuggers/rider.py
new file mode 100644
--- /dev/null
+++ b/Engine/Extras/ushell/channels/unreal/core/debuggers/rider.py
@@ -0,0 +1,84 @@
+# Copyright Epic Games, Inc. All Rights Reserved.
+
+import os
+import subprocess
+from pathlib import Path
+
+#-------------------------------------------------------------------------------
+def _select_sln_path(ue_context):
+    base_dir = ue_context.get_branch() or ue_context.get_project()
+    if not base_dir:
+        return None
+    base_dir = base_dir.get_dir()
+
+    fallbacks = [x for x in base_dir.glob("*.sln") if x.is_file()]
+    fallback = fallbacks[0] if fallbacks else base_dir
+
+    state_dir = base_dir / ".idea"
+    if not state_dir:
+        return fallback
+
+    suffix = ".idea/workspace.xml"
+    sln_states = [x for x in state_dir.glob(".idea.*") if (x / suffix).is_file()]
+    if not sln_states:
+        return fallback
+
+    for sln_state in sorted(sln_states, key=lambda x: (x / suffix).stat().st_mtime):
+        sln_stem = sln_state.stem[6:]
+        candidate = base_dir / (sln_stem + ".sln")
+        if candidate.is_file():
+            return candidate
+
+    return fallback
+
+#-------------------------------------------------------------------------------
+def _find_rider_bin():
+    if os.name != "nt":
+        return Path("rider64")
+
+    root_dir = Path(os.getenv("ProgramFiles", "c:/Program Files/"))
+    root_dir /= "JetBrains"
+    candidates = [x for x in root_dir.glob("JetBrains Rider*") if x.is_dir()]
+    if candidates:
+        candidates.sort(key=lambda x: x.stat().st_mtime)
+        candidate = candidates[-1] / "bin/rider64.exe"
+        if candidate.is_file():
+            return candidate
+
+    return Path("rider64.exe")
+
+#-------------------------------------------------------------------------------
+class Debugger(object):
+    def __init__(self, ue_context):
+        self._ue_context = ue_context
+
+    def debug(self, exec_context, cmd, *args):
+        cmd = exec_context.create_runnable(cmd, *args)
+        cmd.launch(suspended=True, new_term=True)
+        pid = cmd.get_pid()
+
+        try:
+            if not self.attach(pid):
+                cmd.kill()
+        except:
+            cmd.kill()
+            raise
+
+    def attach(self, pid, transport=None, host_ip=None):
+        if transport and transport != "Default":
+            print("Exotic debugger attachment isn't supported!")
+            return False
+
+        bin_path = _find_rider_bin()
+        rider_args = (bin_path, "attach-to-process", str(pid))
+        if sln_path := _select_sln_path(self._ue_context):
+            sln_path = sln_path.absolute()
+            rider_args = (*rider_args, str(sln_path))
+
+        try:
+            print("Launching Rider:", *rider_args)
+            proc = subprocess.Popen(rider_args)
+            return True
+        except Exception as e:
+            print("!! Failed launching Rider:", str(e))
+            return False
diff --git a/Engine/Extras/ushell/channels/unreal/core/describe.flow.py b/Engine/Extras/ushell/channels/unreal/core/describe.flow.py
--- a/Engine/Extras/ushell/channels/unreal/core/describe.flow.py
+++ b/Engine/Extras/ushell/channels/unreal/core/describe.flow.py
@@ -264,6 +264,10 @@ gather.source("cmds/gather.py", "Gather")
 gather.invoke("ushell", "gather")
 
 #-------------------------------------------------------------------------------
+perf_test_auto = flow.describe.Command()
+perf_test_auto.source("cmds/perftest.py", "PerfTestDefault")
+perf_test_auto.invoke("perf", "test", "default")
+
 perf_test_sequence = flow.describe.Command()
 perf_test_sequence.source("cmds/perftest.py", "Sequence")
 perf_test_sequence.invoke("perf", "test", "sequence")
diff --git a/Engine/Extras/ushell/channels/unreal/core/pylib/unreal/_context.py b/Engine/Extras/ushell/channels/unreal/core/pylib/unreal/_context.py
--- a/Engine/Extras/ushell/channels/unreal/core/pylib/unreal/_context.py
+++ b/Engine/Extras/ushell/channels/unreal/core/pylib/unreal/_context.py
@@ -573,8 +573,9 @@ class Context(object):
         primary_dir:Path = Path(primary_dir).resolve()
 
         if sys.platform == "darwin":
-            # On Mac, we want to use the case-correct path, and the fcntl.F_GETPATH operation is the most direct way to get that
-            # Note that if we encounter an exception, the path will be left in its original form
+            # On Mac, we want to use the case-correct path, and the fcntl.F_GETPATH
+            # operation is the most direct way to get that. Note that if we encounter
+            # an exception, the path will be left in its original form
             try:
                 import fcntl
                 with os.fdopen(os.open(primary_dir, os.O_RDONLY)) as primary_dir_object:
@@ -595,6 +596,8 @@ class Context(object):
             while not iter_dir.samefile(iter_dir.anchor):
                 if uproj_path := next(iter_dir.glob("*.uproject*"), None):
                     return self._initialize(uproj_path)
+                if (iter_dir / "Config").is_dir():
+                    return self._initialize_noproject(iter_dir)
                 iter_dir = iter_dir.parent
             raise EnvironmentError(f"Unable to establish an Unreal context from directory '{primary_dir}'")
 
diff --git a/Engine/Extras/ushell/channels/unreal/core/pylib/unreal/_ubt.py b/Engine/Extras/ushell/channels/unreal/core/pylib/unreal/_ubt.py
--- a/Engine/Extras/ushell/channels/unreal/core/pylib/unreal/_ubt.py
+++ b/Engine/Extras/ushell/channels/unreal/core/pylib/unreal/_ubt.py
@@ -307,6 +307,7 @@ class _BuildToolWin64(_BuildToolBase):
             "/property:Configuration=Development;Platform=AnyCPU",
             "/verbosity:m",
             "/nologo",
+            "/restore",
             "/target:Build",
         )
 
