[USHELL] Fix issue with ushell - platform can fail due to reading platform file names without converting to lower case, which causes failures when attempting to get platform info.

diff --git a/Engine/Extras/ushell/channels/unreal/core/pylib/unreal/_context.py b/Engine/Extras/ushell/channels/unreal/core/pylib/unreal/_context.py
--- a/Engine/Extras/ushell/channels/unreal/core/pylib/unreal/_context.py
+++ b/Engine/Extras/ushell/channels/unreal/core/pylib/unreal/_context.py
@@ -539,7 +539,7 @@ class _PlatformProvider(object):
                 continue
 
             for item in extras_dir.glob("platform_*.py"):
-                provisions[item.stem[9:]] = item
+                provisions[item.stem[9:].lower()] = item
 
         return provisions
 
