Updated ushell [43031464]
#rnx

diff --git a/Engine/Extras/ushell/.version b/Engine/Extras/ushell/.version
--- a/Engine/Extras/ushell/.version
+++ b/Engine/Extras/ushell/.version
@@ -1 +1 @@
-42762466
\ No newline at end of file
+43031464
\ No newline at end of file
diff --git a/Engine/Extras/ushell/channels/unreal/core/cmds/uat.py b/Engine/Extras/ushell/channels/unreal/core/cmds/uat.py
--- a/Engine/Extras/ushell/channels/unreal/core/cmds/uat.py
+++ b/Engine/Extras/ushell/channels/unreal/core/cmds/uat.py
@@ -161,4 +161,4 @@ class Uat(unrealcmd.MultiPlatformCmd):
         exec_context = self.get_exec_context()
         cmd = self._get_binary_path()
         cmd = exec_context.create_runnable(cmd, *args)
-        return cmd.run(stdin=False)
+        return cmd.run()
