Updated ushell [43260511]
#rnx

diff --git a/Engine/Extras/ushell/.version b/Engine/Extras/ushell/.version
--- a/Engine/Extras/ushell/.version
+++ b/Engine/Extras/ushell/.version
@@ -1 +1 @@
-43031464
\ No newline at end of file
+43260511
\ No newline at end of file
diff --git a/Engine/Extras/ushell/channels/flow/core/cmds/theme.py b/Engine/Extras/ushell/channels/flow/core/cmds/theme.py
--- a/Engine/Extras/ushell/channels/flow/core/cmds/theme.py
+++ b/Engine/Extras/ushell/channels/flow/core/cmds/theme.py
@@ -152,6 +152,10 @@ def _set_scheme(scheme):
     for i, swatch in enumerate(swatches):
         _send_osc(swatch, 4, i)
 
+    # Selection background
+    sel_bg = bg + (fg - bg).scaled(0.5)
+    _send_osc(sel_bg, 17)
+
 #-------------------------------------------------------------------------------
 def _showcase_scheme(name, scheme):
     import random
diff --git a/Engine/Extras/ushell/channels/flow/core/system/flow/_runnable.py b/Engine/Extras/ushell/channels/flow/core/system/flow/_runnable.py
--- a/Engine/Extras/ushell/channels/flow/core/system/flow/_runnable.py
+++ b/Engine/Extras/ushell/channels/flow/core/system/flow/_runnable.py
@@ -70,7 +70,7 @@ class Runnable(object):
 
     def __iter__(self):
         stdout = self.launch(stdout=True)
-        return (x.decode().rstrip() for x in iter(stdout.readline, b""))
+        return (x.decode().rstrip() for x in stdout)
 
     def get_exe_path(self):
         return self._exe
diff --git a/Engine/Extras/ushell/channels/unreal/core/cmds/stage.py b/Engine/Extras/ushell/channels/unreal/core/cmds/stage.py
--- a/Engine/Extras/ushell/channels/unreal/core/cmds/stage.py
+++ b/Engine/Extras/ushell/channels/unreal/core/cmds/stage.py
@@ -131,7 +131,7 @@ class _Impl(unrealcmd.Cmd):
         if self.args.style == 'nopak':
             uat_packaging_args = None
 
-        cook_flavor = platform.get_cook_flavor()
+        cook_flavor = platform.get_cook_flavor(target_type.name.lower())
 
         dot_uat_cmd = "_uat"
         dot_uat_args = (
diff --git a/Engine/Extras/ushell/channels/unreal/core/pylib/prettyprinter.py b/Engine/Extras/ushell/channels/unreal/core/pylib/prettyprinter.py
--- a/Engine/Extras/ushell/channels/unreal/core/pylib/prettyprinter.py
+++ b/Engine/Extras/ushell/channels/unreal/core/pylib/prettyprinter.py
@@ -4,7 +4,7 @@
 class Printer(object):
     def _thread(self):
         try:
-            for line in iter(self._stdout.readline, b""):
+            for line in self._stdout:
                 self._print(line.decode(errors="replace").rstrip())
         except (IOError, ValueError):
             pass
@@ -19,7 +19,6 @@ class Printer(object):
         self._stdout = runnable.launch(stdout=True)
 
         import threading
-        thread = threading.Thread(target=self._thread)
+        thread = threading.Thread(target=self._thread, daemon=True)
         thread.start()
         runnable.wait()
-        thread.join()
diff --git a/Engine/Extras/ushell/channels/unreal/core/pylib/unreal/_platform.py b/Engine/Extras/ushell/channels/unreal/core/pylib/unreal/_platform.py
--- a/Engine/Extras/ushell/channels/unreal/core/pylib/unreal/_platform.py
+++ b/Engine/Extras/ushell/channels/unreal/core/pylib/unreal/_platform.py
@@ -165,9 +165,9 @@ class Platform(object):
             return form
         raise ValueError(f"No known cook form for target '{target}'")
 
-    def get_cook_flavor(self):
+    def get_cook_flavor(self, target = None):
         if hasattr(self, "_get_cook_flavor"):
-            return self._get_cook_flavor()
+            return self._get_cook_flavor(target)
 
     @_optional_api
     def launch(self, exec_context, stage_dir, binary_path, args):
diff --git a/Engine/Extras/ushell/channels/unreal/core/pylib/unreal/platforms/_android.py b/Engine/Extras/ushell/channels/unreal/core/pylib/unreal/platforms/_android.py
--- a/Engine/Extras/ushell/channels/unreal/core/pylib/unreal/platforms/_android.py
+++ b/Engine/Extras/ushell/channels/unreal/core/pylib/unreal/platforms/_android.py
@@ -5,6 +5,7 @@ import re
 import shutil
 import unreal
 import subprocess as sp
+import glob
 from pathlib import Path
 
 #-------------------------------------------------------------------------------
@@ -99,11 +100,38 @@ class PlatformBase(unreal.Platform):
         except ValueError:
             return
 
-    def _get_cook_form(self, target):
-        if target == "game":   return "Android_ASTC"
-        if target == "client": return "Android_ASTCClient"
+    def _get_cook_form(self, target=None):
+        flavor = self.get_cook_flavor(target)
+
+        if target == "client":
+            return f"Android_{flavor}Client"
+
+        return f"Android_{flavor}"
+
+    def _get_cook_flavor(self, target):
+        cook_dirs = glob.glob(str(self.get_unreal_context().get_project().get_dir() / "Saved/Cooked/Android_*"))
+
+        if target == "game":
+            cook_dirs = [os.path.basename(os.path.normpath(d)) for d in cook_dirs if not d.endswith("Client")]
+            if len(cook_dirs) == 0:
+                return "ASTC"
+
+            if len(cook_dirs) > 1:
+                print(f"Multiple flavors of cook found - using {cook_dirs[0]}")
+
+            return cook_dirs[0][len("Android_"):]
+
+        if target == "client":
+            cook_dirs = [os.path.basename(os.path.normpath(d)) for d in cook_dirs if d.endswith("Client")]
+
+            if len(cook_dirs) == 0:
+                return "ASTC"
+
+            if len(cook_dirs) > 1:
+                print(f"Multiple flavors of cook found - using {cook_dirs[0]}")
+
+            return cook_dirs[0][len("Android_"):-len("Client")]
 
-    def _get_cook_flavor(self):
         return "ASTC"
 
     def _get_stage_name(self, target):
