Updated ushell [43490754]
#rnx

diff --git a/Engine/Extras/ushell/.version b/Engine/Extras/ushell/.version
--- a/Engine/Extras/ushell/.version
+++ b/Engine/Extras/ushell/.version
@@ -1 +1 @@
-43260511
\ No newline at end of file
+43490754
\ No newline at end of file
diff --git a/Engine/Extras/ushell/channels/flow/core/cmds/theme.py b/Engine/Extras/ushell/channels/flow/core/cmds/theme.py
--- a/Engine/Extras/ushell/channels/flow/core/cmds/theme.py
+++ b/Engine/Extras/ushell/channels/flow/core/cmds/theme.py
@@ -144,7 +144,10 @@ def _set_scheme(scheme):
     bg = scheme.get_bg()
     _send_osc(fg, 10)
     _send_osc(bg, 11)
-    _send_osc(fg, 12)
+
+    cursor = _Colour(0xff, 0x00, 0xff) if True else fg
+    _send_osc(cursor, 12)
+
     if not _is_winterm():
         swatches[0] = bg
         swatches[7] = fg
diff --git a/Engine/Extras/ushell/channels/flow/core/system/lib/exeshim/build.bat b/Engine/Extras/ushell/channels/flow/core/system/lib/exeshim/build.bat
--- a/Engine/Extras/ushell/channels/flow/core/system/lib/exeshim/build.bat
+++ b/Engine/Extras/ushell/channels/flow/core/system/lib/exeshim/build.bat
@@ -33,6 +33,6 @@ goto:eof
 :build
 echo.
 echo shim.cpp / %2 / %1
-cl ../shim.cpp /Fo%2.obj /DSHIM_BUILD=%1 /nologo /FC /Zi /Zl /EHs-c- /O1 /c %_cl%
+cl ../shim.cpp /Fo%2.obj /DSHIM_BUILD=%1 /nologo /FC /Zi /Zl /EHs-c- /O1 /GS- /c %_cl%
 link %_link% /nologo /nodefaultlib /align:16 /entry:%2 /out:%2.exe /subsystem:console %2.obj kernel32.lib user32.lib shell32.lib
 goto:eof
diff --git a/Engine/Extras/ushell/channels/flow/core/system/lib/exeshim/shim.cpp b/Engine/Extras/ushell/channels/flow/core/system/lib/exeshim/shim.cpp
--- a/Engine/Extras/ushell/channels/flow/core/system/lib/exeshim/shim.cpp
+++ b/Engine/Extras/ushell/channels/flow/core/system/lib/exeshim/shim.cpp
@@ -70,7 +70,8 @@ Vt100Scope::Vt100Scope()
         return;
 
     _con_mode |= ENABLE_VIRTUAL_TERMINAL_PROCESSING;
-    SetConsoleMode(handle, _con_mode);
+    if (!SetConsoleMode(handle, _con_mode))
+        return;
 
     _stdout_handle = handle;
 }
@@ -81,6 +82,14 @@ Vt100Scope::~Vt100Scope()
     if (intptr_t(_stdout_handle) <= 0)
         return;
 
+    DWORD unused;
+    const wchar_t restore[] =
+        L"\x1b[0 q"             // restore cursor to user's preference
+        L"\x1b" L"1049l"        // switch back to main buffer
+        L"\x1b]9;4;0;50\x07"    // disables progress bar
+        ;
+    WriteFile(_stdout_handle, restore, sizeof(restore) - sizeof(restore[0]), &unused, nullptr);
+
     // Restore modifications to the terminal's mode
     SetConsoleMode(_stdout_handle, _con_mode);
 }
diff --git a/Engine/Extras/ushell/channels/flow/nt/provision.bat b/Engine/Extras/ushell/channels/flow/nt/provision.bat
--- a/Engine/Extras/ushell/channels/flow/nt/provision.bat
+++ b/Engine/Extras/ushell/channels/flow/nt/provision.bat
@@ -7,8 +7,8 @@ goto:eof
 ::------------------------------------------------------------------------------
 :main
 
-set _pyver=3.13.1
-set _pysha=7b7923ff0183a8b8fca90f6047184b419b108cb437f75fc1c002f9d2f8bcec16
+set _pyver=3.13.3
+set _pysha=59ff76e16e6597de47474fb22be69e7191a89116910d728ab735079b078e52db
 set _pytag=313
 set _pymark=%_pyver%.version
 
diff --git a/Engine/Extras/ushell/channels/unreal/core/cmds/build.py b/Engine/Extras/ushell/channels/unreal/core/cmds/build.py
--- a/Engine/Extras/ushell/channels/unreal/core/cmds/build.py
+++ b/Engine/Extras/ushell/channels/unreal/core/cmds/build.py
@@ -179,6 +179,7 @@ class _PrettyPrinter(prettyprinter.Printer):
         out = ""
         if self._section != None:
             out = flow.cmd.text.grey("%3d%%" % self._percent) + " "
+            self.set_progress(self._percent)
 
         source = self._source_re.search(line)
         if source:
diff --git a/Engine/Extras/ushell/channels/unreal/core/cmds/cook.py b/Engine/Extras/ushell/channels/unreal/core/cmds/cook.py
--- a/Engine/Extras/ushell/channels/unreal/core/cmds/cook.py
+++ b/Engine/Extras/ushell/channels/unreal/core/cmds/cook.py
@@ -18,6 +18,7 @@ class _CookPrettyPrinter(uelogprinter.Printer):
         if m:
             percent = (int(m.group(1)) * 100) // max(int(m.group(2)), 1)
             line += f" [{percent}%]"
+            self.set_progress(percent)
         super()._print(line)
 
 
diff --git a/Engine/Extras/ushell/channels/unreal/core/cmds/getbuild.py b/Engine/Extras/ushell/channels/unreal/core/cmds/getbuild.py
--- a/Engine/Extras/ushell/channels/unreal/core/cmds/getbuild.py
+++ b/Engine/Extras/ushell/channels/unreal/core/cmds/getbuild.py
@@ -97,12 +97,14 @@ If you find multiple builds that match you will be given the option to pick one,
     def complete_buildtype(self, prefix):
         return ["packaged", "staged"]
 
+    def is_server_platform(self, platform):
+        server_platforms = ("win64", "windows", "linux", "mac")
+        return platform in server_platforms
+
     def complete_platforms(self, prefix):
-        server_platforms = ("win64", "linux", "mac")
         for platform in unrealcmd.Cmd.complete_platform(self, prefix):
             yield platform
-            yield platform + "-client"
-            if platform in server_platforms:
+            if self.is_server_platform(platform):
                 yield platform + "-server"
 
     def complete_dest(self, prefix):
@@ -112,19 +114,13 @@ If you find multiple builds that match you will be given the option to pick one,
                 yield item.name + "/"
 
     def convert_platform(self, platform):
-        wasClient = False
         wasServer = False
-        if platform.endswith("-client"):
-            platform = platform[0:-7]
-            wasClient = True
         if platform.endswith("-server"):
             platform = platform[0:-7]
             wasServer = True
 
         try:
             convertedPlatform = self.get_platform(platform).get_config_name()
-            if wasClient:
-                return convertedPlatform + "-client"
             if wasServer:
                 return convertedPlatform + "-server"
             return convertedPlatform
@@ -161,13 +157,17 @@ If you find multiple builds that match you will be given the option to pick one,
             return build_id_arg
         return None
 
-    def build_query(self, search_request):
+    def build_query(self, search_request, build_type, branch, platforms):
         search_changelist = None
         if search_request.isnumeric():
             search_changelist = search_request
             print("Searching for build closest to changelist: " + str(search_changelist))
             search_request = "" # we special case handle the search for changelists
 
+        build_id = self.get_build_id(search_request)
+        if build_id:
+            search_request = "" # special case handling the build_id filtering
+
         query = {}
         for statement in search_request.split(";"):
             if statement.find("=") != -1:
@@ -177,14 +177,37 @@ If you find multiple builds that match you will be given the option to pick one,
         if search_changelist:
             query["commit"] = {"$lte": search_changelist}
 
+        if build_id:
+            query["buildId"] = {"$eq": str(build_id)}
+
+        if build_type:
+            query["type"] = {"$eq": str(build_type)}
+
+        if branch:
+            query["branch"] = {"$eq": str(branch.lower())}
+
+        if platforms:
+            platform_filters = []
+
+            for platform in platforms:
+                p = platform.lower()
+                platform_filters.append(p)
+                if not p.endswith("-server"):
+                    platform_filters.append(f"{p}-client")
+
+                if self.is_server_platform(p):
+                    platform_filters.append(f"{p}-server")
+
+            query["platform"] =  {"$in": platform_filters}
+
         return query
 
-    def search_build(self, host, override_host, namespace, bucket, query):
+    def search_build(self, host, override_host, namespace, request):
         # Run a search against the api to find builds that matches the request
         temp_dir = tempfile.TemporaryDirectory()
         query_file_path = f"{temp_dir.name}/query.json"
         with open(query_file_path, "w") as query_file:
-            json.dump(query, query_file)
+            json.dump(request, query_file)
 
         result_file_path = f"{temp_dir.name}/result.json"
 
@@ -193,7 +216,6 @@ If you find multiple builds that match you will be given the option to pick one,
             "list",
             "--host=" + str(host) if host != None else "--override-host=" + str(override_host),
             "--namespace=" + str(namespace),
-            "--bucket=" + bucket,
             "--oidctoken-exe-path=" + self.get_oidc_token_path(),
             "--query-path=" + query_file_path,
             "--result-path=" + result_file_path
@@ -206,7 +228,7 @@ If you find multiple builds that match you will be given the option to pick one,
         with open(result_file_path, "r") as result_fd:
             results = json.load(result_fd)["results"]
 
-        return [{"name" : r["metadata"]["name"], "platforms": {r["metadata"]["platform"]: r["buildId"]}, "commit": r["metadata"].get("commit", "")} for r in results]
+        return [{"name" : r["metadata"]["name"], "bucket": r["bucketId"], "platforms": {r["metadata"]["platform"]: r["buildId"]}, "commit": r["metadata"].get("commit", "")} for r in results]
 
     def _show_select_build_prompt(self, branch_cl, builds):
         closest_cl = 1 << 30
@@ -224,6 +246,9 @@ If you find multiple builds that match you will be given the option to pick one,
                     closest_cl = int(cl)
 
         def abbreviate(x):
+            x = x.replace("-client", "-c")
+            x = x.replace("-server", "-s")
+            x = x.replace("-game", "-g")
             vowels_plus = ("a", "e", "i", "o", "u", "y")
             x = "".join(x for x in x if x not in vowels_plus)
             return x[:7]
@@ -294,7 +319,6 @@ If you find multiple builds that match you will be given the option to pick one,
         if dest == "" and project_path:
             dest = os.path.join(project_path, self.args.buildtype.capitalize())
 
-        build_id = self.get_build_id(self.args.match)
         branch = self._get_p4_stream()
 
         branch_cl = ue_context.get_engine().get_info()["Changelist"]
@@ -320,72 +344,59 @@ If you find multiple builds that match you will be given the option to pick one,
             return False
 
         platforms_to_download = [self.convert_platform(p) for p in self.args.platforms]
-
         builds_to_download = []
-        if build_id:
-            # a single build id was specified, this can only be used wit ha single platform as well
-            if len(platforms_to_download) != 1:
-                self.print_error("Only a single platform can be specified when specifying a build id")
-                return False
-
-            platform = platforms_to_download[0]
-            bucket = f"{project}.{build_type}.{branch}.{platform}".lower()
-            builds_to_download = [(bucket, build_id, platforms_to_download[0])]
-        else:
-            # we need to search
-            found_builds = {}
-
-            query = self.build_query(self.args.match)
-            # Scan all platforms to see which ones have a build that matches
-            for platform in platforms_to_download:
-                self.print_info(f"Enumerating builds for platform {platform}")
-                bucket = f"{project}.{build_type}.{branch}.{platform}".lower()
-
-                builds = self.search_build(host, override_host, namespace, bucket, query)
-                print(f"Found {len(builds)} builds")
-                for build in builds:
-                    build_name = build["name"]
-                    commit_str = str(build.get("commit", ""))
-                    build_key = build_name + commit_str
-                    if build_name in found_builds:
-                        # this build already existed but a new platform was found add it to the list if it didn't already exist
-                        for platform in build["platforms"].keys():
-                            if platform not in found_builds[build_key]["platforms"]:
-                                found_builds[build_key]["platforms"][platform] = build["platforms"][platform]
-                            else:
-                                self.print_info(f"Multiple builds for platform {platform} found for build {build_name}. Using newer build.")
-
+        found_builds = {}
+
+        query = self.build_query(self.args.match, build_type, branch, platforms_to_download)
+        bucket_regex = f"{project}.{build_type}.*".lower()
+        request = {
+            "bucketRegex": bucket_regex,
+            "query": query
+        }
+
+        builds = self.search_build(host, override_host, namespace, request)
+        print(f"Found {len(builds)} builds")
+        for build in builds:
+            build_name = build["name"]
+            commit_str = str(build.get("commit", ""))
+            build_key = build_name + commit_str
+
+            if build_key in found_builds:
+                # this build already existed but a new platform was found add it to the list if it didn't already exist
+                for platform in build["platforms"].keys():
+                    if platform not in found_builds[build_key]["platforms"]:
+                        found_builds[build_key]["platforms"][platform] = build["platforms"][platform]
                     else:
-                        found_builds[build_key] = build
+                        # Multiple builds found but just using the newer build without any output as it can get very spammy
+                        pass
+            else:
+                found_builds[build_key] = build
 
-            # remove builds that only have some of the platforms we want
-            found_builds = {k:v for k, v in found_builds.items() if len(v["platforms"]) == len(platforms_to_download)}
+        if found_builds == None or len(found_builds) == 0:
+            self.print_error(f"Failed to find any build that matches query {request} in namespace {namespace}")
+            return False
 
-            if found_builds == None or len(found_builds) == 0:
-                self.print_error(f"Failed to find any build that matches query {query} in namespace {namespace} and bucket {bucket}")
+        selected_build = None
+        if len(found_builds) == 1:
+            selected_build = next(iter(found_builds.values()))
+            selected_build_name = selected_build["name"]
+            print(f"Single build found, using build {selected_build_name}")
+        else:
+            # Ask the user to select a build
+            print("Select build to fetch:", end="")
+            selected_build_id = self._show_select_build_prompt(branch_cl, found_builds)
+            if not selected_build_id:
+                self.print_error("No build was selected")
                 return False
 
-            selected_build = None
-            if len(found_builds) == 1:
-                selected_build = next(iter(found_builds.values()))
-                selected_build_name = selected_build["name"]
-                print(f"Single build found, using build {selected_build_name}")
-            else:
-                # Ask the user to select a build
-                print("Select build to fetch:", end="")
-                selected_build_id = self._show_select_build_prompt(branch_cl, found_builds)
-                if not selected_build_id:
-                    self.print_error("No build was selected")
-                    return False
-
-                selected_build = found_builds[selected_build_id]
-                selected_build_name = selected_build["name"]
-                print(f"Using build {selected_build_name}")
-
-            # build the selection of build ids to download from which buckets
-            for platform, build_id in selected_build["platforms"].items():
-                bucket = f"{project}.{build_type}.{branch}.{platform}".lower()
-                builds_to_download.extend([(bucket, build_id, platform)])
+            selected_build = found_builds[selected_build_id]
+            selected_build_name = selected_build["name"]
+            print(f"Using build {selected_build_name}")
+
+        # build the selection of build ids to download from which buckets
+        for platform, build_id in selected_build["platforms"].items():
+            bucket = selected_build["bucket"]
+            builds_to_download.extend([(bucket, build_id, platform)])
 
         result = 0
         for bucket, build_id, build_platform in builds_to_download:
diff --git a/Engine/Extras/ushell/channels/unreal/core/pylib/prettyprinter.py b/Engine/Extras/ushell/channels/unreal/core/pylib/prettyprinter.py
--- a/Engine/Extras/ushell/channels/unreal/core/pylib/prettyprinter.py
+++ b/Engine/Extras/ushell/channels/unreal/core/pylib/prettyprinter.py
@@ -2,6 +2,9 @@
 
 #-------------------------------------------------------------------------------
 class Printer(object):
+    def __del__(self):
+        print("\x1b]9;4;0;0\x07", end="")
+
     def _thread(self):
         try:
             for line in self._stdout:
@@ -22,3 +25,7 @@ class Printer(object):
         thread = threading.Thread(target=self._thread, daemon=True)
         thread.start()
         runnable.wait()
+
+    def set_progress(self, percent):
+        percent = int(max(0, min(percent, 100)))
+        print(f"\x1b]9;4;1;{percent}\x07", end="")
