Updated ushell [43608851]
#rnx

diff --git a/Engine/Extras/ushell/.version b/Engine/Extras/ushell/.version
--- a/Engine/Extras/ushell/.version
+++ b/Engine/Extras/ushell/.version
@@ -1 +1 @@
-43490754
\ No newline at end of file
+43608851
\ No newline at end of file
diff --git a/Engine/Extras/ushell/channels/unreal/core/cmds/build.py b/Engine/Extras/ushell/channels/unreal/core/cmds/build.py
--- a/Engine/Extras/ushell/channels/unreal/core/cmds/build.py
+++ b/Engine/Extras/ushell/channels/unreal/core/cmds/build.py
@@ -102,16 +102,12 @@ class _PrettyPrinter(prettyprinter.Printer):
         self._error_file = tempfile.TemporaryFile(mode="w+t")
 
     def print_error_summary(self):
-        import textwrap
-        import shutil
-        msg_width = shutil.get_terminal_size(fallback=(9999,0))[0]
-
         error_file = self._error_file
         error_file.seek(0);
 
         files = {}
         for line in error_file:
-            m = re.search(r"\s*(.+\.[a-z]+)[(:](\d+)[,:)0-9]+\s*(.+)$", line[1:])
+            m = re.search(r"\s*(.+[/\\]\w+\.\w+)[(:](\d+)[,:)0-9]+\s*(.+)$", line[1:])
             if not m:
                 continue
 
@@ -124,6 +120,10 @@ class _PrettyPrinter(prettyprinter.Printer):
         if not files:
             return
 
+        import textwrap
+        import shutil
+        msg_width = shutil.get_terminal_size(fallback=(9999,0))[0]
+
         print()
         print(flow.cmd.text.light_yellow("Errors and warnings:"))
         print()
