Updated ushell [43710683]
#rnx

diff --git a/Engine/Extras/ushell/.version b/Engine/Extras/ushell/.version
--- a/Engine/Extras/ushell/.version
+++ b/Engine/Extras/ushell/.version
@@ -1 +1 @@
-43608851
\ No newline at end of file
+43710683
\ No newline at end of file
diff --git a/Engine/Extras/ushell/channels/unreal/core/pylib/zen/snapshot.py b/Engine/Extras/ushell/channels/unreal/core/pylib/zen/snapshot.py
--- a/Engine/Extras/ushell/channels/unreal/core/pylib/zen/snapshot.py
+++ b/Engine/Extras/ushell/channels/unreal/core/pylib/zen/snapshot.py
@@ -310,6 +310,13 @@ class _Base(zen.cmd.ZenUtilityBaseCmd):
     buildtype   = flow.cmd.Opt("", "Specific build type to get like nightly, standard, etc... (defaults to any type available)")
     build_index = None
 
+    def get_cook_flavor(self):
+        effective_flavor = self.args.flavor
+        if not effective_flavor:
+            platform = self.get_platform(self.args.platform)
+            effective_flavor = platform.get_cook_flavor(self.args.runtime.lower())
+        return effective_flavor
+
     def get_build_index(self):
         if not self.build_index:
             if self.args.nofileshare:
@@ -344,7 +351,7 @@ class _Base(zen.cmd.ZenUtilityBaseCmd):
             branch_name = version["BranchName"]
             project_name = context.get_project().get_name()
 
-            changelists = self.get_build_index().get_available_changelists(project_name, self.args.platform, branch_name, self.args.runtime, self.args.flavor)
+            changelists = self.get_build_index().get_available_changelists(project_name, self.args.platform, branch_name, self.args.runtime, self.get_cook_flavor())
             for changelist in changelists:
                 yield changelist
 
@@ -409,7 +416,7 @@ class Find(_Base):
         branch = self.get_branch_name()
 
         try:
-            found_changelist, builds = self.get_build_index().find_build(project, self.args.platform, branch, self.args.runtime, self.args.flavor, self.args.buildtype, changelist)
+            found_changelist, builds = self.get_build_index().find_build(project, self.args.platform, branch, self.args.runtime, self.get_cook_flavor(), self.args.buildtype, changelist)
 
             if not builds:
                 self.print_error(f"No build found")
@@ -448,7 +455,7 @@ class Get(_Base):
         branch = self.get_branch_name()
 
         try:
-            found_changelist, snapshotdescriptor = self.get_build_index().get_snapshot_descriptor(project_name, self.args.platform, branch, self.args.runtime, self.args.flavor, self.args.buildtype, changelist)
+            found_changelist, snapshotdescriptor = self.get_build_index().get_snapshot_descriptor(project_name, self.args.platform, branch, self.args.runtime, self.get_cook_flavor(), self.args.buildtype, changelist)
 
             while found_changelist != changelist:
                 self.print_error(f"No snapshot found for the specified project({project_name}), branch({branch}), platform({self.args.platform}), runtime({self.args.runtime}), and changelist({changelist})")
@@ -461,7 +468,7 @@ class Get(_Base):
                 if choice == "p":
                     break
 
-                found_changelist, snapshotdescriptor = self.get_build_index().get_snapshot_descriptor(project_name, self.args.platform, version["BranchName"], self.args.runtime, self.args.flavor, self.args.buildtype, changelist)
+                found_changelist, snapshotdescriptor = self.get_build_index().get_snapshot_descriptor(project_name, self.args.platform, version["BranchName"], self.args.runtime, self.get_cook_flavor(), self.args.buildtype, changelist)
 
             if not snapshotdescriptor:
                 self.print_error(f"No build found")
@@ -506,7 +513,7 @@ class List(_Base):
         branch_name = self.get_branch_name()
 
         try:
-            changelists = self.get_build_index().get_available_changelists(project_name, self.args.platform, branch_name, self.args.runtime, self.args.flavor)
+            changelists = self.get_build_index().get_available_changelists(project_name, self.args.platform, branch_name, self.args.runtime, self.get_cook_flavor())
         except BuildIndexException as e:
             self.print_error(e)
             return []
