Updated ushell [43768628]
#rnx

diff --git a/Engine/Extras/ushell/.version b/Engine/Extras/ushell/.version
--- a/Engine/Extras/ushell/.version
+++ b/Engine/Extras/ushell/.version
@@ -1 +1 @@
-43710683
\ No newline at end of file
+43768628
\ No newline at end of file
diff --git a/Engine/Extras/ushell/channels/unreal/core/cmds/cook.py b/Engine/Extras/ushell/channels/unreal/core/cmds/cook.py
--- a/Engine/Extras/ushell/channels/unreal/core/cmds/cook.py
+++ b/Engine/Extras/ushell/channels/unreal/core/cmds/cook.py
@@ -58,15 +58,19 @@ class _Cook(unrealcmd.MultiPlatformCmd):
         # Prepare arguments
         cultures = self.args.cultures.replace(",", "+")
 
+        add_forcerecook = (True for x in self.args.cookargs if x.startswith("-forcerecook="))
+        add_forcerecook = not next(add_forcerecook, False)
+
         commandlet_args = (
             ("-targetplatform=" + cook_form) if cook_form else None,
             ("-cookcultures=" + cultures) if cultures else None,
             "-unattended" if not self.args.nounattended else None,
-            "-unversioned" if not self.versioned else None,
+            "-unversioned" if not self.args.versioned else None,
             "-stdout",
             "-cookonthefly" if self.args.onthefly else None,
             "-iterate" if self.args.iterate else None,
             "-noxgeshadercompile" if self.args.noxge else None,
+            "-forcerecook=false" if add_forcerecook else None,
             *unreal.cmdline.read_ueified(*self.args.cookargs),
         )
 
