Updated ushell [44383714]
#rnx

diff --git a/Engine/Extras/ushell/.version b/Engine/Extras/ushell/.version
--- a/Engine/Extras/ushell/.version
+++ b/Engine/Extras/ushell/.version
@@ -1 +1 @@
-43768628
\ No newline at end of file
+44383714
\ No newline at end of file
diff --git a/Engine/Extras/ushell/channels/unreal/core/cmds/cook.py b/Engine/Extras/ushell/channels/unreal/core/cmds/cook.py
--- a/Engine/Extras/ushell/channels/unreal/core/cmds/cook.py
+++ b/Engine/Extras/ushell/channels/unreal/core/cmds/cook.py
@@ -58,8 +58,8 @@ class _Cook(unrealcmd.MultiPlatformCmd):
         # Prepare arguments
         cultures = self.args.cultures.replace(",", "+")
 
-        add_forcerecook = (True for x in self.args.cookargs if x.startswith("-forcerecook="))
-        add_forcerecook = not next(add_forcerecook, False)
+        forcerecook_alreadypresent = any(True for x in self.args.cookargs \
+            if x == "-forcerecook" or x.startswith("-forcerecook="))
 
         commandlet_args = (
             ("-targetplatform=" + cook_form) if cook_form else None,
@@ -68,9 +68,8 @@ class _Cook(unrealcmd.MultiPlatformCmd):
             "-unversioned" if not self.args.versioned else None,
             "-stdout",
             "-cookonthefly" if self.args.onthefly else None,
-            "-iterate" if self.args.iterate else None,
             "-noxgeshadercompile" if self.args.noxge else None,
-            "-forcerecook=false" if add_forcerecook else None,
+            "-forcerecook=false" if self.args.iterate and not forcerecook_alreadypresent else None,
             *unreal.cmdline.read_ueified(*self.args.cookargs),
         )
 
diff --git a/Engine/Extras/ushell/channels/unreal/core/pylib/unreal/platforms/linux/android.py b/Engine/Extras/ushell/channels/unreal/core/pylib/unreal/platforms/linux/android.py
new file mode 100644
--- /dev/null
+++ b/Engine/Extras/ushell/channels/unreal/core/pylib/unreal/platforms/linux/android.py
@@ -0,0 +1,7 @@
+# Copyright Epic Games, Inc. All Rights Reserved.
+
+from .. import _android
+
+#-------------------------------------------------------------------------------
+class Platform(_android.PlatformBase):
+    pass
