Updated ushell [44478591]
#rnx

diff --git a/Engine/Extras/ushell/.version b/Engine/Extras/ushell/.version
--- a/Engine/Extras/ushell/.version
+++ b/Engine/Extras/ushell/.version
@@ -1 +1 @@
-44383714
\ No newline at end of file
+44478591
\ No newline at end of file
diff --git a/Engine/Extras/ushell/channels/flow/core/system/boot.py b/Engine/Extras/ushell/channels/flow/core/system/boot.py
--- a/Engine/Extras/ushell/channels/flow/core/system/boot.py
+++ b/Engine/Extras/ushell/channels/flow/core/system/boot.py
@@ -31,7 +31,7 @@ def _main():
 
     # Run the '$boot' command
     import run
-    ret = run.main((f"{inst_dir}/manifest", "$boot", *sys.argv[1:]))
+    ret = run.main((str(inst_dir), "$boot", *sys.argv[1:]))
     raise SystemExit(ret)
 
 
diff --git a/Engine/Extras/ushell/channels/flow/core/system/flow/system.py b/Engine/Extras/ushell/channels/flow/core/system/flow/system.py
--- a/Engine/Extras/ushell/channels/flow/core/system/flow/system.py
+++ b/Engine/Extras/ushell/channels/flow/core/system/flow/system.py
@@ -223,12 +223,3 @@ class System(object):
     def get_channel(self, index):
         manifest = self._get_manifest()
         return _Channel(self, manifest["channels"][index])
-
-    def is_path_stale(self, path):
-        manifest_path = self.get_working_dir() + "manifest"
-        try:
-            self._mtime = self._mtime or os.path.getmtime(manifest_path)
-            test_mtime = os.path.getmtime(path)
-            return test_mtime <= self._mtime
-        except FileNotFoundError:
-            return True
diff --git a/Engine/Extras/ushell/channels/flow/core/system/lib/bootstrap.py b/Engine/Extras/ushell/channels/flow/core/system/lib/bootstrap.py
--- a/Engine/Extras/ushell/channels/flow/core/system/lib/bootstrap.py
+++ b/Engine/Extras/ushell/channels/flow/core/system/lib/bootstrap.py
@@ -459,14 +459,13 @@ def _create_shims(manifest):
     builder = shims.Builder(shims_dir)
     builder.clean()
 
-    manifest_path = Path().resolve()
-    manifest_path = manifest_path.parent / manifest_path.name.split("$")[0]
-    manifest_path /= "manifest"
-    manifest_path = str(manifest_path)
+    manifest_dir = Path().resolve()
+    manifest_dir = manifest_dir.parent / manifest_dir.name.split("$")[0]
+    manifest_dir = str(manifest_dir)
 
     # Create command shims
     run_py = os.path.abspath(__file__ + "/../../run.py")
-    args_prefix = ("-Xutf8", "-Esu", run_py, manifest_path)
+    args_prefix = ("-Xutf8", "-Esu", run_py, manifest_dir)
     for name in (x for x in manifest["cmd_tree"] if isinstance(x, str)):
         builder.create_python_shim(name, *args_prefix, name)
 
diff --git a/Engine/Extras/ushell/channels/flow/core/system/run.py b/Engine/Extras/ushell/channels/flow/core/system/run.py
--- a/Engine/Extras/ushell/channels/flow/core/system/run.py
+++ b/Engine/Extras/ushell/channels/flow/core/system/run.py
@@ -43,10 +43,12 @@ def main(args):
     sys.path.append(os.path.abspath(__file__ + "/../"))
     import flow.system
 
-    manifest, *args = args
+    # Determine the core manifest
+    manifest_dir, *args = args
+    manifest = manifest_dir + "/manifest"
 
     # Traverse the command tree using the given arguments
-    working_dir = os.path.dirname(manifest) + "/"
+    working_dir = manifest_dir + "/"
     system = flow.system.System(working_dir)
     cmd_tree = system.get_command_tree()
     tree_node, tail_args = cmd_tree.find(iter(args))
diff --git a/Engine/Extras/ushell/channels/flow/core/system/version b/Engine/Extras/ushell/channels/flow/core/system/version
--- a/Engine/Extras/ushell/channels/flow/core/system/version
+++ b/Engine/Extras/ushell/channels/flow/core/system/version
@@ -1 +1 @@
-23
+24
diff --git a/Engine/Extras/ushell/channels/flow/nt/shell_cmd.py b/Engine/Extras/ushell/channels/flow/nt/shell_cmd.py
--- a/Engine/Extras/ushell/channels/flow/nt/shell_cmd.py
+++ b/Engine/Extras/ushell/channels/flow/nt/shell_cmd.py
@@ -55,10 +55,7 @@ class _Shell(object):
     def _add_clink_lua(self):
         system = self.get_system()
         working_dir = system.get_working_dir()
-
         lua_path = working_dir + "complete.clink.lua"
-        if not system.is_path_stale(lua_path):
-            return
 
         cmd_tree = system.get_command_tree()
         tree_root = cmd_tree.get_root_node()
@@ -70,7 +67,7 @@ class _Shell(object):
         lua.write("}\n")
 
         run_py_path = os.path.abspath(__file__ + "/../../core/system/run.py")
-        cmd = rf'"{sys.executable}" -Xutf8 -Esu "{run_py_path}" "{working_dir}/manifest" $complete --daemon'
+        cmd = rf'"{sys.executable}" -Xutf8 -Esu "{run_py_path}" "{working_dir}" $complete --daemon'
         cmd = cmd.replace("\\", "/")
         lua.write(f"local py_server_cmd = '{cmd}'")
 
diff --git a/Engine/Extras/ushell/channels/flow/nt/shell_pwsh.py b/Engine/Extras/ushell/channels/flow/nt/shell_pwsh.py
--- a/Engine/Extras/ushell/channels/flow/nt/shell_pwsh.py
+++ b/Engine/Extras/ushell/channels/flow/nt/shell_pwsh.py
@@ -30,7 +30,6 @@ class _Shell(object):
         cmd_tree = system.get_command_tree()
         tree_root = cmd_tree.get_root_node()
         run_py_path = os.path.abspath(__file__ + "/../../core/system/run.py")
-        manifest = f"{working_dir}/manifest"
         start_dir = os.getcwd()
 
         with open(cookie, "wt") as out:
@@ -90,7 +89,7 @@ function Get-Daemon {{
             "-Xutf8",
             "-Esu",
             "`"{run_py_path}`"",
-            "`"{working_dir}/manifest`"",
+            "`"{working_dir}`"",
             "`$complete",
             "--daemon"
         )
diff --git a/Engine/Extras/ushell/channels/unreal/core/pylib/zen/snapshot.py b/Engine/Extras/ushell/channels/unreal/core/pylib/zen/snapshot.py
--- a/Engine/Extras/ushell/channels/unreal/core/pylib/zen/snapshot.py
+++ b/Engine/Extras/ushell/channels/unreal/core/pylib/zen/snapshot.py
@@ -343,22 +343,7 @@ class _Base(zen.cmd.ZenUtilityBaseCmd):
 
         return self.build_index
 
-    def complete_changelist(self, prefix):
-        context = self.get_unreal_context()
-        engine = context.get_engine()
-        version = engine.get_info()
-        if version:
-            branch_name = version["BranchName"]
-            project_name = context.get_project().get_name()
-
-            changelists = self.get_build_index().get_available_changelists(project_name, self.args.platform, branch_name, self.args.runtime, self.get_cook_flavor())
-            for changelist in changelists:
-                yield changelist
-
-    def complete_runtime(self, prefix):
-        yield "client"
-        yield "server"
-        yield "game"
+    complete_runtime = ("client", "server", "game")
 
     def complete_platform(self, prefix):
         context = self.get_unreal_context()
