Updated ushell [44866953]
#rnx

diff --git a/Engine/Extras/ushell/.version b/Engine/Extras/ushell/.version
--- a/Engine/Extras/ushell/.version
+++ b/Engine/Extras/ushell/.version
@@ -1 +1 @@
-44478591
\ No newline at end of file
+44866953
\ No newline at end of file
diff --git a/Engine/Extras/ushell/channels/flow/core/cmds/complete.py b/Engine/Extras/ushell/channels/flow/core/cmds/complete.py
--- a/Engine/Extras/ushell/channels/flow/core/cmds/complete.py
+++ b/Engine/Extras/ushell/channels/flow/core/cmds/complete.py
@@ -5,17 +5,6 @@ import sys
 import flow.cmd
 import importlib.util as imputil
 
-#-------------------------------------------------------------------------------
-def _set_title_inner(t, s):
-    print("\x1b]2;", t, " - ushell", s, sep="", file=sys.stdout)
-
-if os.name == "nt":
-    def _set_title(t):
-        _set_title_inner(t, "\x1b\x5c")
-else:
-    def _set_title(t):
-        _set_title_inner(t, "\x07")
-
 #-------------------------------------------------------------------------------
 def _reply(line):   sys.stdout.write(line + "\n")
 def _reply_ok():    sys.stdout.write("\x01\n")
@@ -48,7 +37,9 @@ def _get_prompt(system, line_iter, reply):
 
     if title := (os.getenv("FLOW_TITLE") or getattr(context, "TITLE", None)):
         title = expand_tags(title)
-        _set_title(title)
+        print(f"\x1b]0;{title} - ushell\x07", file=sys.stderr)
+
+    print("\x1b[0 q", file=sys.stderr) # cursor reset
 
     for line in line_iter:
         line = expand_tags(line)
diff --git a/Engine/Extras/ushell/channels/flow/core/system/flow/describe.py b/Engine/Extras/ushell/channels/flow/core/system/flow/describe.py
--- a/Engine/Extras/ushell/channels/flow/core/system/flow/describe.py
+++ b/Engine/Extras/ushell/channels/flow/core/system/flow/describe.py
@@ -7,7 +7,7 @@ class Tool(object):
     class _Bundle(object):
         def __init__(self):                       self._payloads = []; self._bin_paths = []
         def platform(self, platform):             self._platform = platform
-        def bin(self, bin_path):                  self._bin_paths.append(bin_path)
+        def bin(self, bin_path, *, alias=None):   self._bin_paths.append((bin_path, alias))
         def sha1(self, sha1):                     self._sha1 = sha1
         def payload(self, url, *, filetype=None): self._payloads.append((url, filetype))
 
@@ -28,7 +28,7 @@ class Tool(object):
     def root_dir(self, root_dir):             self._root_dir = root_dir
     def payload(self, url, *, filetype=None): self.bundle().payload(url, filetype=filetype)
     def platform(self, platform):             self.bundle().platform(platform)
-    def bin(self, bin_path):                  self.bundle().bin(bin_path)
+    def bin(self, bin_path, *, alias=None):   self.bundle().bin(bin_path, alias=alias)
     def sha1(self, sha1):                     self.bundle().sha1(sha1)
 
     # Implemented elsewhere.
diff --git a/Engine/Extras/ushell/channels/flow/core/system/flow/system.py b/Engine/Extras/ushell/channels/flow/core/system/flow/system.py
--- a/Engine/Extras/ushell/channels/flow/core/system/flow/system.py
+++ b/Engine/Extras/ushell/channels/flow/core/system/flow/system.py
@@ -145,7 +145,7 @@ class _Tool(_FlowObject):
         return self._data["version"]
 
     def get_bin_paths(self):
-        return self._data["bin_paths"]
+        return tuple(x[0] for x in self._data["bin_paths"])
 
 
 
diff --git a/Engine/Extras/ushell/channels/flow/core/system/lib/bootstrap.py b/Engine/Extras/ushell/channels/flow/core/system/lib/bootstrap.py
--- a/Engine/Extras/ushell/channels/flow/core/system/lib/bootstrap.py
+++ b/Engine/Extras/ushell/channels/flow/core/system/lib/bootstrap.py
@@ -270,7 +270,7 @@ def _install_tool(name, tool, cleaner):
 
     tool_double = f"{name}-{tool._version}/"
     tool_dir = "../tools/" + tool_double
-    manifest_path = tool_dir + "manifest.1.flow"
+    manifest_path = tool_dir + "manifest.2.flow"
     _log.write("Installing tool", tool_double)
 
     # If the tool exists then reuse the cached manifest
@@ -317,12 +317,12 @@ def _install_tool(name, tool, cleaner):
     for bundle in manifest["bundles"].values():
         if not bundle["enabled"]:
             continue
-        for bin_path in bundle["bin_paths"]:
+        for bin_path, alias in bundle["bin_paths"]:
             if not os.path.exists(tool_dir + bin_path):
                 _log.write(f"Unable to find '{bin_path}' in '{tool_dir}'")
                 continue
 
-            bin_paths.add(bin_path)
+            bin_paths.add((bin_path, alias))
 
             # Confirm we have a binary that is executable.
             if os.name == "nt":
@@ -476,10 +476,10 @@ def _create_shims(manifest):
             version = tool["version"]
             bin_paths = tool["bin_paths"]
             tool_dir = tool_root + tool["double"] + "/"
-            for bin_path in bin_paths:
-                name = os.path.basename(bin_path)
+            for bin_path, alias in bin_paths:
                 tool_path = tool_dir + bin_path
                 if os.path.isfile(tool_path):
+                    name = os.path.basename(alias or bin_path)
                     builder.create_shim(name, tool_path)
                 else:
                     _log.write("Binary not found;", tool_path)
diff --git a/Engine/Extras/ushell/channels/flow/posix/provision.sh b/Engine/Extras/ushell/channels/flow/posix/provision.sh
--- a/Engine/Extras/ushell/channels/flow/posix/provision.sh
+++ b/Engine/Extras/ushell/channels/flow/posix/provision.sh
@@ -66,26 +66,48 @@ function _check_version()
     cat << EOF | $1 -EsSB
 import sys
 vi = sys.version_info
-raise SystemExit(vi.major != ${py_ver_maj} or vi.minor != ${py_ver_min})
+raise SystemExit(vi.major != ${py_ver_maj} or vi.minor < ${py_ver_min})
 EOF
     return $?
 }
 
+function _check_bin()
+{
+    if ! command -v $1; then
+        return 1
+    fi
+
+    if ! _check_version $1; then
+        return 1
+    fi
+
+    return 0
+}
+
 function _bin_select()
 {
     py_bin=
-    for candidate in "$@"; do
-        if ! command -v $candidate; then
-            continue
-        fi
 
-        if ! _check_version $candidate; then
+    if command -v uv; then
+        for candidate in `uv python find --no-project`; do
+            if ! _check_bin $candidate; then
+                continue
+            fi
+
+            py_bin=$candidate
+            return 0
+        done
+    fi
+
+    for candidate in "$@"; do
+        if ! _check_bin $candidate; then
             continue
         fi
 
         py_bin=$candidate
         return 0
     done
+
     return 1
 }
 
@@ -100,10 +122,11 @@ _bin_select \
     python
 
 if [ -z "$py_bin" ]; then
-    _error "No suitable binary found that is $py_ver"
+    _error "No suitable binary found that is $py_ver (or better)"
     echo ""
     echo "Examples to install Python $py_ver:"
     echo ""
+    echo "       uv: uv python install python$py_ver"
     echo "   Ubuntu: apt install python$py_ver"
     echo "      Mac: brew install python@$py_ver"
     echo ""
diff --git a/Engine/Extras/ushell/channels/unreal/perforce/pylib/p4utils/__init__.py b/Engine/Extras/ushell/channels/unreal/perforce/pylib/p4utils/__init__.py
--- a/Engine/Extras/ushell/channels/unreal/perforce/pylib/p4utils/__init__.py
+++ b/Engine/Extras/ushell/channels/unreal/perforce/pylib/p4utils/__init__.py
@@ -29,6 +29,19 @@ def get_p4_set(prop_name):
 
     return p4_set_result.get(prop_name, None)
 
+#-------------------------------------------------------------------------------
+def determine_server(info=None):
+    if not info:
+        info = P4.info().run()
+
+    ret = "perforce:1666"
+    for key in ("changeServer", "proxyAddress", "serverAddress"):
+        if key in info:
+            ret = getattr(info, key)
+            break
+
+    return ret
+
 #-------------------------------------------------------------------------------
 def login():
     try:
diff --git a/Engine/Extras/ushell/channels/unreal/perforce/pylib/peafour.py b/Engine/Extras/ushell/channels/unreal/perforce/pylib/peafour.py
--- a/Engine/Extras/ushell/channels/unreal/perforce/pylib/peafour.py
+++ b/Engine/Extras/ushell/channels/unreal/perforce/pylib/peafour.py
@@ -14,10 +14,13 @@ class _P4Result(object):
         super().__setattr__("_result", result)
 
     def __str__(self):                 return str(self._result)
-    def __contains__(self, key):       return (key in self._result) or (key + "0" in self._result)
     def as_dict(self):                 return {k.decode():v.decode() for k,v in self._result.items()}
     def __setattr__(self, key, value): self._result[key.encode()] = str(value).encode()
 
+    def __contains__(self, key):
+        key = key.encode()
+        return (key in self._result) or (key + b"0" in self._result)
+
     def __getattr__(self, key):
         if (key + "0").encode() in self._result:
             def as_list():
