Updated ushell [45062782]
#rnx

diff --git a/Engine/Extras/ushell/.version b/Engine/Extras/ushell/.version
--- a/Engine/Extras/ushell/.version
+++ b/Engine/Extras/ushell/.version
@@ -1 +1 @@
-44866953
\ No newline at end of file
+45062782
\ No newline at end of file
diff --git a/Engine/Extras/ushell/channels/flow/nt/cmd_thefuzz.bat b/Engine/Extras/ushell/channels/flow/nt/cmd_thefuzz.bat
--- a/Engine/Extras/ushell/channels/flow/nt/cmd_thefuzz.bat
+++ b/Engine/Extras/ushell/channels/flow/nt/cmd_thefuzz.bat
@@ -53,7 +53,7 @@ goto:eof
 set _base="%~f2"
 if "%~2"=="" set _base=\
 echo %~1 ^^
-fd -d8 -td . %_base:\"="% | %_fzf% --scheme=default --history=%temp%/thefuzz_chdir
+fd --follow -d8 -td . %_base:\"="% | %_fzf% --scheme=default --history=%temp%/thefuzz_chdir
 goto:eof
 
 ::------------------------------------------------------------------------------
@@ -61,5 +61,5 @@ goto:eof
 set _base="%~f1"
 if "%~1"=="" set _base=\
 echo explorer.exe ^^
-fd -d8 -td . %_base:\"="% | %_fzf% --scheme=default --history=%temp%/thefuzz_chdir
+fd --follow -d8 -td . %_base:\"="% | %_fzf% --scheme=default --history=%temp%/thefuzz_chdir
 goto:eof
diff --git a/Engine/Extras/ushell/channels/unreal/core/cmds/getbuild.py b/Engine/Extras/ushell/channels/unreal/core/cmds/getbuild.py
--- a/Engine/Extras/ushell/channels/unreal/core/cmds/getbuild.py
+++ b/Engine/Extras/ushell/channels/unreal/core/cmds/getbuild.py
@@ -19,24 +19,29 @@ Unless otherwise specified it is downloaded to <projectdir>/Staged/<platform>.
 
 If your search only matches one build that build is downloaded for you, to add filtering you can use --match as follows:
     * --match 8v93f39cvn3           - If you know the build id:
-    * --match 1337                  - To search for the build cloest to a changelist
+    * --match 1337                  - To search for the build closest to a changelist
     * --match job=v99f3;step=39f983 - For advanced searching you can specify key values
 
 If you find multiple builds that match you will be given the option to pick one, were the * indicates the build closest to your current commit."""
 
-    buildtype   = unrealcmd.Arg(str, "The type of build to download. Usually packaged or staged builds.")
-    platforms   = unrealcmd.Arg([str], "The platform(s) you want to download from a build")
-    runargs     = unrealcmd.Arg([str], "Additional arguments to pass to the process being launched")
+    buildtype           = unrealcmd.Arg(str, "The type of build to download. Usually packaged or staged builds.")
+    platforms           = unrealcmd.Arg([str], "The platform(s) you want to download from a build")
+    runargs             = unrealcmd.Arg([str], "Additional arguments to pass to the process being launched")
 
-    project     = unrealcmd.Opt("", "The project you want to download a build for. Defaults to the current project in ushell.")
-    branch      = unrealcmd.Opt("", "Branch/Stream to fetch data for, defaults to current p4 stream")
-    namespace   = unrealcmd.Opt("", "The namespace to use overriding what is defined in engine ini")
-    host        = unrealcmd.Opt("", "The host name of the server to connect to, defaults to your cloud storage server in engine ini")
-    clean       = unrealcmd.Opt(False, "Ignore any previously downloaded data and fetch it all again")
-    verbose     = unrealcmd.Opt(False, "Enable verbose logging output")
+    project             = unrealcmd.Opt("", "The project you want to download a build for. Defaults to the current project in ushell.")
+    branch              = unrealcmd.Opt("", "Branch/Stream to fetch data for, defaults to current p4 stream")
+    namespace           = unrealcmd.Opt("", "The namespace to use overriding what is defined in engine ini")
+    host                = unrealcmd.Opt("", "The host name of the server to connect to, defaults to your cloud storage server in engine ini")
+    clean               = unrealcmd.Opt(False, "Ignore any previously downloaded data and fetch it all again")
+    verbose             = unrealcmd.Opt(False, "Enable verbose logging output")
+    wildcard            = unrealcmd.Opt("", "Windows style wildcard (using * and ?) to match file paths to include")
 
-    dest        = unrealcmd.Opt("", "The directory to download the build to. Local contents will be overwritten. Defaults to `<projectdir>/Staged/<platform>`")
-    match       = unrealcmd.Opt("", "Either a build id (hexadecimal string), a changelist (number) or a key value seperate list of key=values you want to find a build for. e.g. changelist=1337;job=foo")
+    dest                = unrealcmd.Opt("", "The directory to download the build to. Local contents will be overwritten. Defaults to `<projectdir>/Staged/<platform>`")
+    match               = unrealcmd.Opt("", "Either a build id (hexadecimal string), a changelist (number) or a key value separate list of key=values you want to find a build for. e.g. changelist=1337;job=foo")
+
+    def __init__(self):
+        setattr(type(self), "exclude-wildcard", unrealcmd.Opt("", "Windows style wildcard (using * and ?) to match file paths to exclude. Applied after --wildcard include filter"))
+        super().__init__()
 
     def lookup_service_name(self):
         config = self.get_unreal_context().get_config()
@@ -228,7 +233,7 @@ If you find multiple builds that match you will be given the option to pick one,
         with open(result_file_path, "r") as result_fd:
             results = json.load(result_fd)["results"]
 
-        return [{"name" : r["metadata"]["name"], "bucket": r["bucketId"], "platforms": {r["metadata"]["platform"]: r["buildId"]}, "commit": r["metadata"].get("commit", "")} for r in results]
+        return [{"name" : r["metadata"]["name"], "bucket": r["bucketId"], "platforms": {r["metadata"]["platform"]: r["buildId"]}, "buildgroup" : r["metadata"].get("buildgroup", ""), "commit": r["metadata"].get("commit", "")} for r in results]
 
     def _show_select_build_prompt(self, branch_cl, builds):
         closest_cl = 1 << 30
@@ -257,7 +262,10 @@ If you find multiple builds that match you will be given the option to pick one,
         def list_available(closest_cl, available_builds):
             for build in available_builds:
                 cl = build["commit"]
-                name = build["name"]
+                # group on build group first if present
+                name = build.get("buildgroup", "")
+                if not name:
+                    name = build["name"]
                 platforms = build["platforms"].keys()
                 middle = ""
                 for column_platform,short_name in columns.items():
@@ -357,7 +365,9 @@ If you find multiple builds that match you will be given the option to pick one,
         builds = self.search_build(host, override_host, namespace, request)
         print(f"Found {len(builds)} builds")
         for build in builds:
-            build_name = build["name"]
+            build_name = build.get("buildgroup", None)
+            if not build_name:
+                build_name = build["name"]
             commit_str = str(build.get("commit", ""))
             build_key = build_name + commit_str
 
@@ -415,7 +425,9 @@ If you find multiple builds that match you will be given the option to pick one,
                 "--clean" if self.args.clean else "",
                 "--verbose" if self.args.verbose else "",
                 "--local-path=\"" + build_dest + "\"",
-                "--build-id=" + build_id
+                "--build-id=" + build_id,
+                "--wildcard={self.args.wildcard}" if self.args.wildcard else "",
+                "--exclude-wildcard=" + getattr(self.args, "exclude-wildcard", "")
             )
 
             result = self.run_zen(args)
diff --git a/Engine/Extras/ushell/channels/unreal/core/pylib/unreal/platforms/mac/visionos.py b/Engine/Extras/ushell/channels/unreal/core/pylib/unreal/platforms/mac/visionos.py
new file mode 100644
--- /dev/null
+++ b/Engine/Extras/ushell/channels/unreal/core/pylib/unreal/platforms/mac/visionos.py
@@ -0,0 +1,24 @@
+# Copyright Epic Games, Inc. All Rights Reserved.
+
+import os
+import unreal
+
+#-------------------------------------------------------------------------------
+class Platform(unreal.Platform):
+    name = "VISIONOS"
+
+    def _read_env(self):
+        yield from ()
+
+    def _get_version_ue4(self):
+        dot_cs = self.get_unreal_context().get_engine().get_dir()
+        dot_cs /= "Source/Programs/UnrealBuildTool/Platform/Mac/MacToolChain.cs"
+        return Platform._get_version_helper_ue4(dot_cs, "MinMacOSVersion")
+
+    def _get_version_ue5(self):
+        dot_cs = "Source/Programs/UnrealBuildTool/Platform/Mac/ApplePlatformSDK"
+        version = self._get_version_helper_ue5(dot_cs + ".Versions.cs")
+        return version or self._get_version_helper_ue5(dot_cs + ".cs")
+
+    def _get_cook_form(self, target):
+        if target == "game":   return "VISIONOS"
diff --git a/Engine/Extras/ushell/channels/unreal/core/pylib/zen/snapshot.py b/Engine/Extras/ushell/channels/unreal/core/pylib/zen/snapshot.py
--- a/Engine/Extras/ushell/channels/unreal/core/pylib/zen/snapshot.py
+++ b/Engine/Extras/ushell/channels/unreal/core/pylib/zen/snapshot.py
@@ -19,7 +19,7 @@ class BuildIndexException(Exception):
 #-------------------------------------------------------------------------------
 class BuildIndex(ABC):
     @abstractmethod
-    def get_available_changelists(self, project, platform, branch, runtime, flavor):
+    def get_available_changelists(self, project, platform, branch, runtime, flavor, buildtype = ""):
         pass
 
     @abstractmethod
@@ -112,8 +112,8 @@ class UnrealCloudDDCBuildIndex(BuildIndex):
 
         return results
 
-    def get_available_changelists(self, project, platform, branch, runtime, flavor):
-        builds = self.query_builds(project, platform, branch, runtime, flavor)
+    def get_available_changelists(self, project, platform, branch, runtime, flavor, buildtype):
+        builds = self.query_builds(project, platform, branch, runtime, flavor, buildtype)
         return list(map(lambda build : int(build["metadata"]["changelist"]), builds))
 
     def find_build(self, project, platform, branch, runtime, flavor, buildtype, changelist):
@@ -221,12 +221,27 @@ class FileshareBuildIndex(BuildIndex):
 
         return builds_dir
 
-    def get_available_changelists(self, project, platform, branch, runtime, flavor):
+    def get_available_changelists(self, project, platform, branch, runtime, flavor, buildtype):
         builds_dir = self.get_builds_dir(project, platform, branch, runtime)
         if builds_dir == None:
             return []
 
-        return [int(dir.name) for dir in os.scandir(builds_dir) if dir.name.isdecimal()]
+        builds = []
+
+        try:
+            changelists_unfiltered = os.scandir(builds_dir)
+            if buildtype:
+                for changelist in changelists_unfiltered:
+                    builds_for_changelist = os.scandir(os.path.join(builds_dir, changelist.name))
+                    build_name_lower = build.name.lower()
+                    if build_name_lower.startswith(buildtype.lower() and changelist.name.isdecimal()):
+                        builds.append(changelist.name)
+            else:
+                builds = [int(dir.name) for dir in changelists_unfiltered if dir.name.isdecimal()]
+        except:
+            pass
+
+        return builds
 
     def get_filtered_builds(self, project, platform, branch, runtime, buildtype, changelist):
         builds_dir_for_changelist = os.path.join(self.get_builds_dir(project, platform, branch, runtime), str(changelist))
@@ -498,7 +513,7 @@ class List(_Base):
         branch_name = self.get_branch_name()
 
         try:
-            changelists = self.get_build_index().get_available_changelists(project_name, self.args.platform, branch_name, self.args.runtime, self.get_cook_flavor())
+            changelists = self.get_build_index().get_available_changelists(project_name, self.args.platform, branch_name, self.args.runtime, self.get_cook_flavor(), self.args.buildtype)
         except BuildIndexException as e:
             self.print_error(e)
             return []
