Updated ushell [45231411]
#rnx

diff --git a/Engine/Extras/ushell/.version b/Engine/Extras/ushell/.version
--- a/Engine/Extras/ushell/.version
+++ b/Engine/Extras/ushell/.version
@@ -1 +1 @@
-45062782
\ No newline at end of file
+45231411
\ No newline at end of file
diff --git a/Engine/Extras/ushell/channels/flow/core/system/lib/bootstrap.py b/Engine/Extras/ushell/channels/flow/core/system/lib/bootstrap.py
--- a/Engine/Extras/ushell/channels/flow/core/system/lib/bootstrap.py
+++ b/Engine/Extras/ushell/channels/flow/core/system/lib/bootstrap.py
@@ -221,6 +221,8 @@ def _acquire_tool(name, tool, manifest, target_dir, progress_cb):
         try:
             os.rename(source_dir, target_dir)
             break
+        except FileExistsError:
+            pass
         except OSError:
             if i == 5:
                 raise
@@ -269,7 +271,7 @@ def _install_tool(name, tool, cleaner):
     manifest = _manifest_tool(name, tool)
 
     tool_double = f"{name}-{tool._version}/"
-    tool_dir = "../tools/" + tool_double
+    tool_dir = "../../../../tools/" + tool_double
     manifest_path = tool_dir + "manifest.2.flow"
     _log.write("Installing tool", tool_double)
 
@@ -293,7 +295,8 @@ def _install_tool(name, tool, cleaner):
 
     try:
         _log.write("Acquiring to", tool_dir)
-        _acquire_tool(name, tool, manifest, tool_dir, progress_cb)
+        temp_dir = f"../{tool_double}"
+        _acquire_tool(name, tool, manifest, temp_dir, progress_cb)
     except KeyboardInterrupt:
         raise
     except Exception as e:
@@ -301,8 +304,8 @@ def _install_tool(name, tool, cleaner):
         _log.print("\b" * 4, "fail", sep="", end="\n")
         return manifest
 
-    # Maybe this tool doen't deploy for this platform?
-    if not os.path.isdir(tool_dir):
+    # Maybe this tool doesn't deploy for this platform?
+    if not os.path.isdir(temp_dir):
         _log.print("\b" * 4, " na ", sep="", end="\n")
         return manifest
 
@@ -310,7 +313,7 @@ def _install_tool(name, tool, cleaner):
     if hasattr(tool, "post_install"):
         progress_cb("post")
         _log.write("Post-install")
-        tool.post_install(tool_dir)
+        tool.post_install(temp_dir)
 
     # Validate that the required binaries exist.
     bin_paths = set()
@@ -318,7 +321,7 @@ def _install_tool(name, tool, cleaner):
         if not bundle["enabled"]:
             continue
         for bin_path, alias in bundle["bin_paths"]:
-            if not os.path.exists(tool_dir + bin_path):
+            if not os.path.exists(temp_dir + bin_path):
                 _log.write(f"Unable to find '{bin_path}' in '{tool_dir}'")
                 continue
 
@@ -327,15 +330,26 @@ def _install_tool(name, tool, cleaner):
             # Confirm we have a binary that is executable.
             if os.name == "nt":
                 continue
-            binary = os.path.join(tool_dir, bin_path)
+            binary = os.path.join(temp_dir, bin_path)
             if not (stat.S_IXUSR & os.stat(binary)[stat.ST_MODE]):
                 st = os.stat(binary)
                 os.chmod(binary, st.st_mode | stat.S_IEXEC)
     manifest["bin_paths"] = tuple(bin_paths)
 
+    manifest_path = temp_dir + os.path.basename(manifest_path)
     with open(manifest_path, "wb") as out:
         marshal.dump(manifest, out)
 
+    # Slot into place
+    try:
+        os.rename(temp_dir, tool_dir)
+    except FileExistsError:
+        pass
+    except Exception as e:
+        _log.write("ERROR", str(e))
+        _log.print("\b" * 4, "fail", sep="", end="\n")
+        return manifest
+
     _log.print("\b" * 4, " ok ", sep="", end="\n")
     return manifest
 
@@ -460,7 +474,12 @@ def _create_shims(manifest):
     builder.clean()
 
     manifest_dir = Path().resolve()
-    manifest_dir = manifest_dir.parent / manifest_dir.name.split("$")[0]
+    for item in manifest_dir.parents:
+        if item.name == "$cleaner":
+            manifest_dir = item.parent / manifest_dir.name
+            break
+    else:
+        raise EnvironmentError("Could not find '$cleaner' root directory")
     manifest_dir = str(manifest_dir)
 
     # Create command shims
@@ -470,7 +489,7 @@ def _create_shims(manifest):
         builder.create_python_shim(name, *args_prefix, name)
 
     # Create tool shims
-    tool_root = os.path.abspath("../tools") + "/"
+    tool_root = os.path.abspath("../../../../tools") + "/"
     for channel in reversed(manifest["channels"]):
         for name, tool in channel["tools"].items():
             version = tool["version"]
@@ -545,6 +564,12 @@ class _Cleaner(object):
         self._dir.mkdir(parents=True, exist_ok=True)
         self._count = 0
 
+    def create_temp_dir(self, name:str) -> Path:
+        ret = self._dir / str(self._count) / name
+        ret.mkdir(parents=True)
+        self._count += 1
+        return ret
+
     def __del__(self) -> None:
         shutil.rmtree(self._dir.parent, ignore_errors=True)
 
@@ -708,16 +733,16 @@ def impl() -> Path:
 
     dest_dir = state_dir.get_dir()
 
+    (dest_dir.parent / "tools").mkdir(parents=True, exist_ok=True)
+
     cleaner = _Cleaner(working_dir)
     cleaner.delete(dest_dir)
 
-    temp_dir = dest_dir.parent / (dest_dir.name + f"${os.getpid()}")
+    temp_dir = cleaner.create_temp_dir(dest_dir.name)
     try:
-        temp_dir.mkdir(parents=True)
         _update(temp_dir, sources, cleaner)
         dest_dir = temp_dir.rename(dest_dir)
-    except:
-        cleaner.delete(temp_dir)
-        raise
+    except FileExistsError:
+        pass
 
     return dest_dir
