Updated ushell [45361181]

#rnx

diff --git a/Engine/Extras/ushell/.version b/Engine/Extras/ushell/.version
--- a/Engine/Extras/ushell/.version
+++ b/Engine/Extras/ushell/.version
@@ -1 +1 @@
-45231411
\ No newline at end of file
+45361181
\ No newline at end of file
diff --git a/Engine/Extras/ushell/channels/unreal/core/cmds/getbuild.py b/Engine/Extras/ushell/channels/unreal/core/cmds/getbuild.py
--- a/Engine/Extras/ushell/channels/unreal/core/cmds/getbuild.py
+++ b/Engine/Extras/ushell/channels/unreal/core/cmds/getbuild.py
@@ -29,14 +29,14 @@ If you find multiple builds that match you will be given the option to pick one,
     runargs             = unrealcmd.Arg([str], "Additional arguments to pass to the process being launched")
 
     project             = unrealcmd.Opt("", "The project you want to download a build for. Defaults to the current project in ushell.")
-    branch              = unrealcmd.Opt("", "Branch/Stream to fetch data for, defaults to current p4 stream")
+    branch              = unrealcmd.Opt("", "Branch/Stream to fetch data for, defaults to current p4 stream. Accepts a p4 depot path like //UE5/Main or its converted form ue5-main")
     namespace           = unrealcmd.Opt("", "The namespace to use overriding what is defined in engine ini")
     host                = unrealcmd.Opt("", "The host name of the server to connect to, defaults to your cloud storage server in engine ini")
     clean               = unrealcmd.Opt(False, "Ignore any previously downloaded data and fetch it all again")
     verbose             = unrealcmd.Opt(False, "Enable verbose logging output")
     wildcard            = unrealcmd.Opt("", "Windows style wildcard (using * and ?) to match file paths to include")
 
-    dest                = unrealcmd.Opt("", "The directory to download the build to. Local contents will be overwritten. Defaults to `<projectdir>/Staged/<platform>`")
+    dest                = unrealcmd.Opt("", "The directory to download the build to. Local contents will be overwritten. Defaults to `<projectdir>/Saved/<StagedBuilds/PackagedBuilds>/<platform>`")
     match               = unrealcmd.Opt("", "Either a build id (hexadecimal string), a changelist (number) or a key value separate list of key=values you want to find a build for. e.g. changelist=1337;job=foo")
 
     def __init__(self):
@@ -141,11 +141,23 @@ If you find multiple builds that match you will be given the option to pick one,
         else:
             return build_type
 
+    def convert_buildtype_foldername(self, build_type):
+        if build_type == "packaged":
+            return "PackagedBuilds"
+        elif build_type == "staged":
+            return "StagedBuilds"
+        else:
+            return build_type.capitalize()
+
+    def sanitize_branch_name(self, branch):
+        # massage the perforce stream name to match expectations, no initial / or trailing / and - separation between depot and stream as well as replacing any . with -
+        return branch.lstrip("/").replace("/", "-").replace('.', '-')
+
     def _get_p4_stream(self):
         self.print_info("Checking Perforce stream")
         if branch := self.args.branch:
             print("User provided:", branch)
-            return branch
+            return self.sanitize_branch_name(branch)
 
         branch_dir = self.get_unreal_context().get_branch().get_dir()
         info = P4(d=str(branch_dir)).info().run()
@@ -154,8 +166,7 @@ If you find multiple builds that match you will be given the option to pick one,
         if not branch:
             raise EnvironmentError("Unable to get Perforce client name")
         print("Using stream", branch)
-        # massage the perforce stream name to match expectations, no initial / and - seperation between depot and stream
-        return branch.lstrip("/").replace("/", "-")
+        return self.sanitize_branch_name(branch)
 
     def get_build_id(self, build_id_arg):
         if len(build_id_arg) == 24 and all(char in string.hexdigits for char in build_id_arg):
@@ -325,7 +336,7 @@ If you find multiple builds that match you will be given the option to pick one,
         build_type = self.convert_buildtype(self.args.buildtype)
         dest = self.args.dest
         if dest == "" and project_path:
-            dest = os.path.join(project_path, self.args.buildtype.capitalize())
+            dest = os.path.join(project_path, "Saved", self.convert_buildtype_foldername(self.args.buildtype))
 
         branch = self._get_p4_stream()
 
