Allow odsc from ushell to generate shader symbols through the --symbols argument

diff --git a/Engine/Extras/ushell/channels/unreal/core/cmds/odsc.py b/Engine/Extras/ushell/channels/unreal/core/cmds/odsc.py
--- a/Engine/Extras/ushell/channels/unreal/core/cmds/odsc.py
+++ b/Engine/Extras/ushell/channels/unreal/core/cmds/odsc.py
@@ -31,6 +31,7 @@ class _Odsc(unrealcmd.MultiPlatformCmd):
     attach   = unrealcmd.Opt(False, "Attach a debugger to the commandlet")
     skipassetscan = unrealcmd.Opt(False, "Skip asset scanning")
     debug    = unrealcmd.Opt(False, "Use debug executables")
+    symbols    = unrealcmd.Opt(False, "Generate shader symbols")
 
     @unrealcmd.Cmd.summarise
     def _odsc(self, cook_form, exec_context=None):
@@ -45,6 +46,11 @@ class _Odsc(unrealcmd.MultiPlatformCmd):
         # Prepare arguments
         cultures = self.args.cultures.replace(",", "+")
 
+        dpcvars="r.ShaderDevelopmentMode=1"
+
+        if self.symbols:
+            dpcvars += ",r.Shaders.Symbols=1"
+
         commandlet_args = (
             "-odsc",
             ("-targetplatform=" + cook_form) if cook_form else None,
@@ -52,7 +58,7 @@ class _Odsc(unrealcmd.MultiPlatformCmd):
             "-unversioned",
             "-stdout",
             "-cookonthefly",
-            "-dpcvars=r.ShaderDevelopmentMode=1",
+            "-dpcvars="+dpcvars,
             "-noxgeshadercompile" if self.args.noxge else None,
             "-skipassetscan" if self.args.skipassetscan else None,
             *unreal.cmdline.read_ueified(*self.args.odscargs),
