fixed shader symbols being generated unconditionally

diff --git a/Engine/Extras/ushell/channels/unreal/core/cmds/odsc.py b/Engine/Extras/ushell/channels/unreal/core/cmds/odsc.py
--- a/Engine/Extras/ushell/channels/unreal/core/cmds/odsc.py
+++ b/Engine/Extras/ushell/channels/unreal/core/cmds/odsc.py
@@ -48,7 +48,7 @@ class _Odsc(unrealcmd.MultiPlatformCmd):
 
         dpcvars="r.ShaderDevelopmentMode=1"
 
-        if self.symbols:
+        if self.args.symbols:
             dpcvars += ",r.Shaders.Symbols=1"
 
         commandlet_args = (
