Updated ushell [45620566]
#rnx

diff --git a/Engine/Extras/ushell/.version b/Engine/Extras/ushell/.version
--- a/Engine/Extras/ushell/.version
+++ b/Engine/Extras/ushell/.version
@@ -1 +1 @@
-45361181
\ No newline at end of file
+45620566
\ No newline at end of file
diff --git a/Engine/Extras/ushell/channels/unreal/core/cmds/getbuild.py b/Engine/Extras/ushell/channels/unreal/core/cmds/getbuild.py
--- a/Engine/Extras/ushell/channels/unreal/core/cmds/getbuild.py
+++ b/Engine/Extras/ushell/channels/unreal/core/cmds/getbuild.py
@@ -7,7 +7,6 @@ import tempfile
 import json
 import os
 from peafour import P4
-import urllib.request
 import string
 import fzf
 
@@ -119,16 +118,21 @@ If you find multiple builds that match you will be given the option to pick one,
                 yield item.name + "/"
 
     def convert_platform(self, platform):
-        wasServer = False
+        was_server = False
         if platform.endswith("-server"):
             platform = platform[0:-7]
-            wasServer = True
+            was_server = True
+
+        # non-windows platforms do not load the windows platform but we still need to translate the platform name
+        # as such we hardcode the conversion for this platform
+        if platform == "win64":
+            platform = "windows"
 
         try:
-            convertedPlatform = self.get_platform(platform).get_config_name()
-            if wasServer:
-                return convertedPlatform + "-server"
-            return convertedPlatform
+            converted_platform = self.get_platform(platform).get_config_name()
+            if was_server:
+                return converted_platform + "-server"
+            return converted_platform
         except ValueError:
             # unknown platform, just assume the name it valid
             return platform
@@ -191,7 +195,7 @@ If you find multiple builds that match you will be given the option to pick one,
                 query[k] = {"$eq": v}
 
         if search_changelist:
-            query["commit"] = {"$lte": search_changelist}
+            query["commit"] = {"$lte": int(search_changelist)}
 
         if build_id:
             query["buildId"] = {"$eq": str(build_id)}
@@ -244,7 +248,13 @@ If you find multiple builds that match you will be given the option to pick one,
         with open(result_file_path, "r") as result_fd:
             results = json.load(result_fd)["results"]
 
-        return [{"name" : r["metadata"]["name"], "bucket": r["bucketId"], "platforms": {r["metadata"]["platform"]: r["buildId"]}, "buildgroup" : r["metadata"].get("buildgroup", ""), "commit": r["metadata"].get("commit", "")} for r in results]
+        return [{
+            "name" : r["metadata"]["name"],
+            "bucket": r["bucketId"],
+            "platforms": {r["metadata"].get("cookPlatform", r["metadata"]["platform"]): r["buildId"]}, # use cook platform over platform if present as that field follows UE conventions for how the platform is called
+            "buildgroup" : r["metadata"].get("buildgroup", ""),
+            "commit": r["metadata"].get("commit", "")
+        } for r in results]
 
     def _show_select_build_prompt(self, branch_cl, builds):
         closest_cl = 1 << 30
diff --git a/Engine/Extras/ushell/channels/unreal/core/cmds/info.py b/Engine/Extras/ushell/channels/unreal/core/cmds/info.py
--- a/Engine/Extras/ushell/channels/unreal/core/cmds/info.py
+++ b/Engine/Extras/ushell/channels/unreal/core/cmds/info.py
@@ -128,7 +128,13 @@ class Info(InfoBase):
             platform = self.get_platform(platform_name)
 
             inner["name"] = platform.get_name()
-            inner["version"] = platform.get_version() or "!unknown"
+
+            version = platform.get_version(detailed=True) or "!unknown"
+            if isinstance(version, dict):
+                for key, value in version.items():
+                    inner["version." + key] = value
+            else:
+                inner["version"] = version
 
             inner["env"] = {}
             for item in platform.read_env():
diff --git a/Engine/Extras/ushell/channels/unreal/core/pylib/unreal/_platform.py b/Engine/Extras/ushell/channels/unreal/core/pylib/unreal/_platform.py
--- a/Engine/Extras/ushell/channels/unreal/core/pylib/unreal/_platform.py
+++ b/Engine/Extras/ushell/channels/unreal/core/pylib/unreal/_platform.py
@@ -74,7 +74,7 @@ class Platform(object):
     def get_unreal_context(self):
         return self._context
 
-    def _get_version_json(self):
+    def _get_version_json(self, *, key="MainVersion", use_autosdk=True):
         def conform_json(inp):
             # As the time of writing the SDK.json files have trailing commas
             # and C-style comments. So we'll try and conform it
@@ -105,15 +105,16 @@ class Platform(object):
                 except:
                     return
 
-            # patch the AutoSDK name while we're here...
-            if autosdk_name := sdk_info.get("AutoSDKPlatform", None):
-                type(self).autosdk_name = autosdk_name
+            if use_autosdk:
+                # patch the AutoSDK name while we're here...
+                if autosdk_name := sdk_info.get("AutoSDKPlatform", None):
+                    type(self).autosdk_name = autosdk_name
 
-            # fetch the version value
-            if version := sdk_info.get("AutoSDKDirectory", None):
-                return version
+                # fetch the version value
+                if version := sdk_info.get("AutoSDKDirectory", None):
+                    return version
 
-            if version := sdk_info.get("MainVersion", None):
+            if version := sdk_info.get(key, None):
                 return version
 
             # ...or traverse up the platform forest
@@ -129,11 +130,15 @@ class Platform(object):
         json_path = engine_dir / f"Config/{name}/{name}_SDK.json"
         return read_sdk_json(json_path)
 
-    def get_version(self):
+    def get_version(self, *, detailed=False):
         if not (version := self._version or self._get_version_json()):
             engine_version = self._context.get_engine().get_version_major()
             self._version = self._get_version_ue5() if engine_version > 4 else self._get_version_ue4()
             version = self._version or ""
+
+        if isinstance(version, dict) and version and not detailed:
+            version = next(x for x in version.values())
+
         return version
 
     def read_env(self):
diff --git a/Engine/Extras/ushell/channels/unreal/core/pylib/unreal/platforms/win64/win64.py b/Engine/Extras/ushell/channels/unreal/core/pylib/unreal/platforms/win64/win64.py
--- a/Engine/Extras/ushell/channels/unreal/core/pylib/unreal/platforms/win64/win64.py
+++ b/Engine/Extras/ushell/channels/unreal/core/pylib/unreal/platforms/win64/win64.py
@@ -106,6 +106,24 @@ class _PlatformBase(unreal.Platform):
 
         return max(sdk_versions) + "-" + max(msvc_versions)
 
+    def _get_version_json(self, **kwargs):
+        ret = {
+            "sdk" : super()._get_version_json()
+        }
+
+        toolchains = (
+            ("vc",      "PreferredVisualCppVersions"),
+            ("clang",   "PreferredClangVersions"),
+        )
+        for name, key in toolchains:
+            tool_ver = super()._get_version_json(key=key, use_autosdk=False)
+            if not tool_ver:                continue
+            if isinstance(tool_ver, list):  tool_ver = max(x for x in tool_ver)
+            if "-" in tool_ver:             tool_ver = tool_ver.split("-")[0]
+            ret[name] = tool_ver
+
+        return ret
+
 #-------------------------------------------------------------------------------
 class Platform(_PlatformBase):
     name = "Win64"
