Updated ushell [45721939]

#rnx

diff --git a/Engine/Extras/ushell/.version b/Engine/Extras/ushell/.version
--- a/Engine/Extras/ushell/.version
+++ b/Engine/Extras/ushell/.version
@@ -1 +1 @@
-45620566
\ No newline at end of file
+45721939
\ No newline at end of file
diff --git a/Engine/Extras/ushell/channels/flow/core/system/lib/bootstrap.py b/Engine/Extras/ushell/channels/flow/core/system/lib/bootstrap.py
--- a/Engine/Extras/ushell/channels/flow/core/system/lib/bootstrap.py
+++ b/Engine/Extras/ushell/channels/flow/core/system/lib/bootstrap.py
@@ -544,7 +544,7 @@ def _finalise(manifests):
     primary = {
         "channels"   : _carve_channels(manifests.values()),
         "pylib_dirs" : tuple(pylib_dirs),
-        "tools_dir"  : os.path.abspath("../tools") + "/",
+        "tools_dir"  : os.path.abspath("../../../../tools") + "/",
         "cmd_tree"   : _plant_cmd_tree(manifests.values()),
     }
 
diff --git a/Engine/Extras/ushell/channels/flow/core/system/version b/Engine/Extras/ushell/channels/flow/core/system/version
--- a/Engine/Extras/ushell/channels/flow/core/system/version
+++ b/Engine/Extras/ushell/channels/flow/core/system/version
@@ -1 +1 @@
-24
+25
diff --git a/Engine/Extras/ushell/channels/flow/nt/describe.flow.py b/Engine/Extras/ushell/channels/flow/nt/describe.flow.py
--- a/Engine/Extras/ushell/channels/flow/nt/describe.flow.py
+++ b/Engine/Extras/ushell/channels/flow/nt/describe.flow.py
@@ -14,9 +14,9 @@ clink.bin("clink_x64.exe")
 
 #-------------------------------------------------------------------------------
 fd = flow.describe.Tool()
-fd.version("8.7.1")
+fd.version("10.3.0")
 fd.payload("[REDACTED_PROFILE_URL]/fd/releases/download/v$VERSION/fd-v$VERSION-x86_64-pc-windows-msvc.zip")
-fd.sha1("7e8f5d0a9d1fed75b12483809f9619c04b39b60c")
+fd.sha1("7b1aed3c07b66f78d18f7dc78bcfcace65ed0d63")
 fd.platform("win32")
 fd.bin("fd.exe")
 fd.source("[REDACTED_PROFILE_URL]/fd/releases/latest", "fd-v([0-9.]+)-x86_64")
diff --git a/Engine/Extras/ushell/channels/unreal/core/cmds/sln.py b/Engine/Extras/ushell/channels/unreal/core/cmds/sln.py
--- a/Engine/Extras/ushell/channels/unreal/core/cmds/sln.py
+++ b/Engine/Extras/ushell/channels/unreal/core/cmds/sln.py
@@ -28,6 +28,7 @@ def _make_temp_vcxproj(ctx, path, name, target, ubt_args=""):
     p.add_import("$(VCTargetsPath)/Microsoft.Cpp.props")
     p.add_import("$(VCTargetsPath)/Microsoft.Cpp.targets")
 
+    path.parent.mkdir(parents=True, exist_ok=True)
     with open(path, "wt") as out:
         p.write(out)
 
diff --git a/Engine/Extras/ushell/channels/unreal/core/describe.flow.py b/Engine/Extras/ushell/channels/unreal/core/describe.flow.py
--- a/Engine/Extras/ushell/channels/unreal/core/describe.flow.py
+++ b/Engine/Extras/ushell/channels/unreal/core/describe.flow.py
@@ -22,24 +22,24 @@ if bundle := fzf.bundle(platform="darwin-arm64"):
 
 #-------------------------------------------------------------------------------
 ripgrep = flow.describe.Tool()
-ripgrep.version("14.1.0")
+ripgrep.version("14.1.1")
 ripgrep.source("[REDACTED_PROFILE_URL]/ripgrep/releases/latest", r"ripgrep-(\d+\.\d+\.\d+)-x86_64")
 if bundle := ripgrep.bundle(platform="win32"):
     bundle.payload("[REDACTED_PROFILE_URL]/ripgrep/releases/download/$VERSION/ripgrep-$VERSION-x86_64-pc-windows-msvc.zip")
     bundle.bin("rg.exe")
-    bundle.sha1("ba7820c05aeb49050d9242a109e422585c3fc767")
+    bundle.sha1("b4ffb31dd4ffff1cfbbd5f7efaeccc88d2ed4162")
 if bundle := ripgrep.bundle(platform="linux"):
     bundle.payload("[REDACTED_PROFILE_URL]/ripgrep/releases/download/$VERSION/ripgrep-$VERSION-x86_64-unknown-linux-musl.tar.gz")
     bundle.bin("rg")
-    bundle.sha1("8267841999fa96f5a5a375ed518a24bebef6ca4c")
+    bundle.sha1("4fac9ca4ab5b22c9bb3ba5984c12d79de86de6b4")
 if bundle := ripgrep.bundle(platform="darwin-amd64"):
     bundle.payload("[REDACTED_PROFILE_URL]/ripgrep/releases/download/$VERSION/ripgrep-$VERSION-x86_64-apple-darwin.tar.gz")
     bundle.bin("rg")
-    bundle.sha1("0d98d699f08c4eed18f4379734d4e1d973f59cde")
+    bundle.sha1("81fb373949dc25ddd2486f63b0701900f9f1fc12")
 if bundle := ripgrep.bundle(platform="darwin-arm64"):
     bundle.payload("[REDACTED_PROFILE_URL]/ripgrep/releases/download/$VERSION/ripgrep-$VERSION-aarch64-apple-darwin.tar.gz")
     bundle.bin("rg")
-    bundle.sha1("59d0d543a7e25831b29ebe2943da366d4c005d26")
+    bundle.sha1("32f8675774041138541db0588fc828137b79b05c")
 
 
 
diff --git a/Engine/Extras/ushell/channels/unreal/core/pylib/unreal/_platform.py b/Engine/Extras/ushell/channels/unreal/core/pylib/unreal/_platform.py
--- a/Engine/Extras/ushell/channels/unreal/core/pylib/unreal/_platform.py
+++ b/Engine/Extras/ushell/channels/unreal/core/pylib/unreal/_platform.py
@@ -80,7 +80,7 @@ class Platform(object):
             # and C-style comments. So we'll try and conform it
             ret = ""
             for line in inp:
-                if m := re.match(r'^(.*[^"])//', line):
+                if m := re.match(r'^(.*[^:"])//', line):
                     line = m.group(1)
                 line = line.strip()
                 if not line:
diff --git a/Engine/Extras/ushell/channels/unreal/core/pylib/unreal/platforms/mac/visionos.py b/Engine/Extras/ushell/channels/unreal/core/pylib/unreal/platforms/mac/visionos.py
--- a/Engine/Extras/ushell/channels/unreal/core/pylib/unreal/platforms/mac/visionos.py
+++ b/Engine/Extras/ushell/channels/unreal/core/pylib/unreal/platforms/mac/visionos.py
@@ -5,20 +5,17 @@ import unreal
 
 #-------------------------------------------------------------------------------
 class Platform(unreal.Platform):
-    name = "VISIONOS"
+    name = "VisionOS"
 
     def _read_env(self):
         yield from ()
 
     def _get_version_ue4(self):
-        dot_cs = self.get_unreal_context().get_engine().get_dir()
-        dot_cs /= "Source/Programs/UnrealBuildTool/Platform/Mac/MacToolChain.cs"
-        return Platform._get_version_helper_ue4(dot_cs, "MinMacOSVersion")
+        return
 
     def _get_version_ue5(self):
-        dot_cs = "Source/Programs/UnrealBuildTool/Platform/Mac/ApplePlatformSDK"
-        version = self._get_version_helper_ue5(dot_cs + ".Versions.cs")
-        return version or self._get_version_helper_ue5(dot_cs + ".cs")
+        return
 
     def _get_cook_form(self, target):
-        if target == "game":   return "VISIONOS"
+        if target == "game": return "VisionOS"
+        if target == "client": return "VisionOSClient"
