Updated ushell [45905185]
#rnx

diff --git a/Engine/Extras/ushell/.version b/Engine/Extras/ushell/.version
--- a/Engine/Extras/ushell/.version
+++ b/Engine/Extras/ushell/.version
@@ -1 +1 @@
-45721939
\ No newline at end of file
+45905185
\ No newline at end of file
diff --git a/Engine/Extras/ushell/channels/unreal/core/cmds/getbuild.py b/Engine/Extras/ushell/channels/unreal/core/cmds/getbuild.py
--- a/Engine/Extras/ushell/channels/unreal/core/cmds/getbuild.py
+++ b/Engine/Extras/ushell/channels/unreal/core/cmds/getbuild.py
@@ -250,8 +250,7 @@ If you find multiple builds that match you will be given the option to pick one,
 
         return [{
             "name" : r["metadata"]["name"],
-            "bucket": r["bucketId"],
-            "platforms": {r["metadata"].get("cookPlatform", r["metadata"]["platform"]): r["buildId"]}, # use cook platform over platform if present as that field follows UE conventions for how the platform is called
+            "platforms": {r["metadata"].get("cookPlatform", r["metadata"]["platform"]): (r["buildId"], r["bucketId"])}, # use cook platform over platform if present as that field follows UE conventions for how the platform is called
             "buildgroup" : r["metadata"].get("buildgroup", ""),
             "commit": r["metadata"].get("commit", "")
         } for r in results]
@@ -410,7 +409,7 @@ If you find multiple builds that match you will be given the option to pick one,
         selected_build = None
         if len(found_builds) == 1:
             selected_build = next(iter(found_builds.values()))
-            selected_build_name = selected_build["name"]
+            selected_build_name = selected_build.get("buildgroup", selected_build["name"])
             print(f"Single build found, using build {selected_build_name}")
         else:
             # Ask the user to select a build
@@ -421,12 +420,11 @@ If you find multiple builds that match you will be given the option to pick one,
                 return False
 
             selected_build = found_builds[selected_build_id]
-            selected_build_name = selected_build["name"]
+            selected_build_name = selected_build.get("buildgroup", selected_build["name"])
             print(f"Using build {selected_build_name}")
 
         # build the selection of build ids to download from which buckets
-        for platform, build_id in selected_build["platforms"].items():
-            bucket = selected_build["bucket"]
+        for platform, (build_id, bucket) in selected_build["platforms"].items():
             builds_to_download.extend([(bucket, build_id, platform)])
 
         result = 0
diff --git a/Engine/Extras/ushell/channels/unreal/core/cmds/stage.py b/Engine/Extras/ushell/channels/unreal/core/cmds/stage.py
--- a/Engine/Extras/ushell/channels/unreal/core/cmds/stage.py
+++ b/Engine/Extras/ushell/channels/unreal/core/cmds/stage.py
@@ -12,6 +12,7 @@ class _Impl(unrealcmd.Cmd):
     style    = unrealcmd.Arg("auto", "Type of result to stage")
     uatargs  = unrealcmd.Arg([str], "Additional arguments to pass to UAT")
     debug    = unrealcmd.Opt(False, "Use a debug variant of UAT")
+    attach   = unrealcmd.Opt(False, "Attach a debugger to UAT")
 
     complete_target = ("game", "client", "server")
 
@@ -137,6 +138,7 @@ class _Impl(unrealcmd.Cmd):
         dot_uat_args = (
             "BuildCookRun",
             "--debug" if self.args.debug else None,
+            "--attach" if self.args.attach else None,
             "--", # NB. NOT a UAT arg!
             #"-project=" + ue_context.get_project().get_path() <-- implicit via _uat
             target_arg,
diff --git a/Engine/Extras/ushell/channels/unreal/core/cmds/uat.py b/Engine/Extras/ushell/channels/unreal/core/cmds/uat.py
--- a/Engine/Extras/ushell/channels/unreal/core/cmds/uat.py
+++ b/Engine/Extras/ushell/channels/unreal/core/cmds/uat.py
@@ -114,10 +114,13 @@ class Uat(unrealcmd.MultiPlatformCmd):
         if not unconventional:
             return self._get_build_script()
 
+        import platform
+        arch = "x64" if platform.machine() == "AMD64" else "Arm64"
+
         variant = variant or "Debug"
         ue_context = self.get_unreal_context()
         engine_dir = ue_context.get_engine().get_dir()
-        return engine_dir / f"Source/Programs/AutomationTool/bin/{variant}/AutomationTool.exe"
+        return engine_dir / f"Source/Programs/AutomationTool/bin/{arch}/{variant}/AutomationTool.exe"
 
     def _launch_debugger(self, *args):
         self.print_info("Launching Visual Studio")
