Updated ushell [46307652]
#rnx

diff --git a/Engine/Extras/ushell/.version b/Engine/Extras/ushell/.version
--- a/Engine/Extras/ushell/.version
+++ b/Engine/Extras/ushell/.version
@@ -1 +1 @@
-45905185
\ No newline at end of file
+46307652
\ No newline at end of file
diff --git a/Engine/Extras/ushell/channels/flow/core/cmds/complete.py b/Engine/Extras/ushell/channels/flow/core/cmds/complete.py
--- a/Engine/Extras/ushell/channels/flow/core/cmds/complete.py
+++ b/Engine/Extras/ushell/channels/flow/core/cmds/complete.py
@@ -139,7 +139,8 @@ class Complete(flow.cmd.Cmd):
     """ Provides support for adding completion to host shells and for debugging
     cmds' completion routines. Suffix the last word of 'words' with "..." to use
     that as a prefix. Use "$" to print the prompt """
-    words   = flow.cmd.Arg([str], "Words to parse and print completion for")
+    command = flow.cmd.Arg("", "The initiating command")
+    words   = flow.cmd.Arg([str], "Follow on words to parse and print completion for")
     daemon  = flow.cmd.Opt(False, "Run as a daemon serving requests over stdin/out")
     delim   = flow.cmd.Opt(" ", "The character that delimits the last word")
 
@@ -159,7 +160,10 @@ class Complete(flow.cmd.Cmd):
         if self.args.daemon:
             return self._main_daemon(system)
 
-        words = self.args.words
+        words = tuple()
+        if command := self.args.command:
+            words = (command, *self.args.words)
+
         if words and words[-1].endswith("..."):
             word_iter = (*words[:-1], words[-1][:-3])
         else:
diff --git a/Engine/Extras/ushell/channels/flow/posix/shells.py b/Engine/Extras/ushell/channels/flow/posix/shells.py
--- a/Engine/Extras/ushell/channels/flow/posix/shells.py
+++ b/Engine/Extras/ushell/channels/flow/posix/shells.py
@@ -91,7 +91,7 @@ def _get_complete_bash():
 function _flow_prompt() {
     if [[ ! -z "$FLOW_PROMPT" ]]; then
         local prompt=${FLOW_PROMPT/:FLOW_CWD:/$(dirs +0)}
-        PS1=$(\$complete \$ "$prompt")
+        PS1=$(\$complete \$ -- "$prompt")
     fi
 }
 if [[ ! -z "$FLOW_PROMPT" ]]; then
@@ -99,7 +99,7 @@ if [[ ! -z "$FLOW_PROMPT" ]]; then
 fi
 
 function _flow_complete() {
-    local choices=$(\$complete ${COMP_WORDS[@]}...)
+    local choices=$(\$complete ${COMP_WORDS[0]} -- ${COMP_WORDS[@]:1}...)
     COMPREPLY=($(compgen -W "$choices" -- ${COMP_WORDS[$COMP_CWORD]}))
 }
 complete -F _flow_complete""" # purposely left no trailing CRLF
@@ -169,7 +169,7 @@ fi
 
 function _flow_complete() {
     # setopt local_options xtrace
-    compadd - $('$complete' ${^^words[@]:0:-1} $PREFIX...)
+    compadd - $('$complete' ${words[1]} -- ${words[@]:1:-1} $PREFIX...)
 }
 compdef _flow_complete '$$$FLOW_CMDS$$$'
 """
diff --git a/Engine/Extras/ushell/channels/unreal/core/cmds/clangdb.py b/Engine/Extras/ushell/channels/unreal/core/cmds/clangdb.py
--- a/Engine/Extras/ushell/channels/unreal/core/cmds/clangdb.py
+++ b/Engine/Extras/ushell/channels/unreal/core/cmds/clangdb.py
@@ -37,6 +37,9 @@ class ClangDb(unrealcmd.Cmd):
     ubtargs   = unrealcmd.Arg([str], "Arguments to pass onwards to UnrealBuildTool")
     platform  = unrealcmd.Opt("", "Platform to build a database for")
     projected = unrealcmd.Opt(False, "Add the -Project= argument when running UnrealBuildTool")
+    keepgenfiles        = unrealcmd.Opt(False, "Add compile commands for UHT generated files (cache build time will increase)")
+    keepplatformcode    = unrealcmd.Opt(False, "Add compile commands for Platforms files (cache build time will increase)")
+    keepthirdpartycode  = unrealcmd.Opt(False, "Add compile commands for ThirdParty files (cache build time will increase)")
 
     @unrealcmd.Cmd.summarise
     def _run_ubt(self, target, platform, ubt_args):
@@ -88,12 +91,27 @@ class ClangDb(unrealcmd.Cmd):
         print("Target:", target.get_name())
         print("Platform:", platform)
 
+        excluded = []
+        if not self.args.keepgenfiles:
+            excluded.append(".../Intermediate/")
+        if not self.args.keepplatformcode:
+            excluded.append(".../Platforms/")
+        if not self.args.keepthirdpartycode:
+            excluded.append(".../ThirdParty/")
+
         ubt_args = (
             "-allmodules",
             "-Mode=GenerateClangDatabase",
+            f"-Exclude={";".join(excluded)}" if excluded else "",
             *self.args.ubtargs,
         )
 
+        db_dir = ue_context.get_engine().get_dir().parent
+        if ue_context.get_type() == unreal.ContextType.FOREIGN:
+            if project := ue_context.get_project():
+                db_dir = project.get_dir()
+                ubt_args = ("-OutputDir=" + str(db_dir), *ubt_args)
+
         if target.is_project_target() or self.args.projected:
             if project := ue_context.get_project():
                 ubt_args = ("-Project=" + str(project.get_path()), *ubt_args)
@@ -105,7 +123,7 @@ class ClangDb(unrealcmd.Cmd):
             return ret
 
         # Print some information
-        db_path = ue_context.get_engine().get_dir() / "../compile_commands.json"
+        db_path = db_dir / "compile_commands.json"
         try: db_size = os.path.getsize(db_path)
         except: db_size = -1
         print()
diff --git a/Engine/Extras/ushell/channels/unreal/core/cmds/getbuild.py b/Engine/Extras/ushell/channels/unreal/core/cmds/getbuild.py
--- a/Engine/Extras/ushell/channels/unreal/core/cmds/getbuild.py
+++ b/Engine/Extras/ushell/channels/unreal/core/cmds/getbuild.py
@@ -33,6 +33,7 @@ If you find multiple builds that match you will be given the option to pick one,
     host                = unrealcmd.Opt("", "The host name of the server to connect to, defaults to your cloud storage server in engine ini")
     clean               = unrealcmd.Opt(False, "Ignore any previously downloaded data and fetch it all again")
     verbose             = unrealcmd.Opt(False, "Enable verbose logging output")
+    preflight           = unrealcmd.Opt(False, "If set, shows preflights")
     wildcard            = unrealcmd.Opt("", "Windows style wildcard (using * and ?) to match file paths to include")
 
     dest                = unrealcmd.Opt("", "The directory to download the build to. Local contents will be overwritten. Defaults to `<projectdir>/Saved/<StagedBuilds/PackagedBuilds>/<platform>`")
@@ -177,7 +178,7 @@ If you find multiple builds that match you will be given the option to pick one,
             return build_id_arg
         return None
 
-    def build_query(self, search_request, build_type, branch, platforms):
+    def build_query(self, search_request, build_type, branch, platforms, preflight):
         search_changelist = None
         if search_request.isnumeric():
             search_changelist = search_request
@@ -206,6 +207,9 @@ If you find multiple builds that match you will be given the option to pick one,
         if branch:
             query["branch"] = {"$eq": str(branch.lower())}
 
+        if not preflight:
+            query["ispreflight"] = {"$eq": "false"}
+
         if platforms:
             platform_filters = []
 
@@ -252,7 +256,8 @@ If you find multiple builds that match you will be given the option to pick one,
             "name" : r["metadata"]["name"],
             "platforms": {r["metadata"].get("cookPlatform", r["metadata"]["platform"]): (r["buildId"], r["bucketId"])}, # use cook platform over platform if present as that field follows UE conventions for how the platform is called
             "buildgroup" : r["metadata"].get("buildgroup", ""),
-            "commit": r["metadata"].get("commit", "")
+            "commit": r["metadata"].get("commit", ""),
+            "hordetemplate" : r["metadata"].get("hordetemplateid", "")
         } for r in results]
 
     def _show_select_build_prompt(self, branch_cl, builds):
@@ -280,6 +285,10 @@ If you find multiple builds that match you will be given the option to pick one,
         columns = {x:abbreviate(x) for x in active_platforms}
 
         def list_available(closest_cl, available_builds):
+            max_template_name_len = 0
+            for build in available_builds:
+                max_template_name_len = max(max_template_name_len, len(build["hordetemplate"]))
+
             for build in available_builds:
                 cl = build["commit"]
                 # group on build group first if present
@@ -291,6 +300,7 @@ If you find multiple builds that match you will be given the option to pick one,
                 for column_platform,short_name in columns.items():
                     middle += " "
                     middle += short_name if column_platform in platforms else (" " * len(short_name))
+                middle += f" {build["hordetemplate"]:<{max_template_name_len}}" if build["hordetemplate"] else (" " * (max_template_name_len + 1))
                 if cl != None:
                     if not p4_change_found:
                         closest_cl = branch_cl
@@ -375,7 +385,7 @@ If you find multiple builds that match you will be given the option to pick one,
         builds_to_download = []
         found_builds = {}
 
-        query = self.build_query(self.args.match, build_type, branch, platforms_to_download)
+        query = self.build_query(self.args.match, build_type, branch, platforms_to_download, self.args.preflight)
         bucket_regex = f"{project}.{build_type}.*".lower()
         request = {
             "bucketRegex": bucket_regex,
diff --git a/Engine/Extras/ushell/channels/unreal/perforce/cmds/clean.py b/Engine/Extras/ushell/channels/unreal/perforce/cmds/clean.py
--- a/Engine/Extras/ushell/channels/unreal/perforce/cmds/clean.py
+++ b/Engine/Extras/ushell/channels/unreal/perforce/cmds/clean.py
@@ -306,7 +306,7 @@ class Reset(flow.cmd.Cmd):
         root_dir = branch.get_dir()
         ret = {
             # root_dir / "*", # to much scope to go wrong here and it's so few files
-            root_dir / "Template/...",
+            root_dir / "Templates/...",
             root_dir / "Engine/...",
         }
         for uproj_path in branch.read_projects():
diff --git a/Engine/Extras/ushell/channels/unreal/perforce/pylib/p4utils/__init__.py b/Engine/Extras/ushell/channels/unreal/perforce/pylib/p4utils/__init__.py
--- a/Engine/Extras/ushell/channels/unreal/perforce/pylib/p4utils/__init__.py
+++ b/Engine/Extras/ushell/channels/unreal/perforce/pylib/p4utils/__init__.py
@@ -174,7 +174,11 @@ class TempBranchSpec(object):
         P4.branch(i=True).run(input_data=spec, on_error=False)
 
     def __del__(self):
-        P4.branch(self._name, d=True).run()
+        sp.run(
+            ("p4", "branch", "-d", self._name),
+            stdout=sp.DEVNULL,
+            stderr=sp.DEVNULL
+        )
 
     def __str__(self):
         return self._name
