Updated ushell [46831677]
#rnx

diff --git a/Engine/Extras/ushell/.version b/Engine/Extras/ushell/.version
--- a/Engine/Extras/ushell/.version
+++ b/Engine/Extras/ushell/.version
@@ -1 +1 @@
-46307652
\ No newline at end of file
+46831677
\ No newline at end of file
diff --git a/Engine/Extras/ushell/channels/flow/core/system/flow/cmd.py b/Engine/Extras/ushell/channels/flow/core/system/flow/cmd.py
--- a/Engine/Extras/ushell/channels/flow/core/system/flow/cmd.py
+++ b/Engine/Extras/ushell/channels/flow/core/system/flow/cmd.py
@@ -93,6 +93,7 @@ class _Env(object):
                 yield key, None
 
     def __iter__(self): return iter(self.keys())
+    def __len__(self):  return len(self._env)
     def keys(self):     return self._env.keys()
     def values(self):   return self._env.values()
     def items(self):    return self._env.items()
diff --git a/Engine/Extras/ushell/channels/unreal/core/cmds/build.py b/Engine/Extras/ushell/channels/unreal/core/cmds/build.py
--- a/Engine/Extras/ushell/channels/unreal/core/cmds/build.py
+++ b/Engine/Extras/ushell/channels/unreal/core/cmds/build.py
@@ -90,6 +90,7 @@ class _Builder(object):
 #-------------------------------------------------------------------------------
 class _PrettyPrinter(prettyprinter.Printer):
     def __init__(self, logger):
+        super().__init__()
         self._source_re = re.compile(r"^(\s+\[[0-9/]+\]\s+|)([^ ]+\.[a-z]+)")
         self._error_re = re.compile("([Ee]rror:|ERROR|Exception: |error LNK|error C)")
         self._warning_re = re.compile("(warning [CV]|[Ww]arning: |note:)")
diff --git a/Engine/Extras/ushell/channels/unreal/core/cmds/getbuild.py b/Engine/Extras/ushell/channels/unreal/core/cmds/getbuild.py
--- a/Engine/Extras/ushell/channels/unreal/core/cmds/getbuild.py
+++ b/Engine/Extras/ushell/channels/unreal/core/cmds/getbuild.py
@@ -450,12 +450,12 @@ If you find multiple builds that match you will be given the option to pick one,
                 "--host=" + str(host) if host != None else "--override-host=" + str(override_host),
                 "--namespace=" + str(namespace),
                 "--bucket=" + bucket,
-                "--oidctoken-exe-path=\"" + self.get_oidc_token_path() + "\"",
+                "--oidctoken-exe-path=" + self.get_oidc_token_path(),
                 "--clean" if self.args.clean else "",
                 "--verbose" if self.args.verbose else "",
-                "--local-path=\"" + build_dest + "\"",
+                "--local-path=" + build_dest,
                 "--build-id=" + build_id,
-                "--wildcard={self.args.wildcard}" if self.args.wildcard else "",
+                f"--wildcard={self.args.wildcard}" if self.args.wildcard else "",
                 "--exclude-wildcard=" + getattr(self.args, "exclude-wildcard", "")
             )
 
diff --git a/Engine/Extras/ushell/channels/unreal/core/cmds/odsc.py b/Engine/Extras/ushell/channels/unreal/core/cmds/odsc.py
--- a/Engine/Extras/ushell/channels/unreal/core/cmds/odsc.py
+++ b/Engine/Extras/ushell/channels/unreal/core/cmds/odsc.py
@@ -18,6 +18,7 @@ class _OdscPrettyPrinter(uelogprinter.Printer):
         if m:
             percent = (int(m.group(1)) * 100) // max(int(m.group(2)), 1)
             line += f" [{percent}%]"
+            self.set_progress(percent)
         super()._print(line)
 
 
diff --git a/Engine/Extras/ushell/channels/unreal/core/cmds/project.py b/Engine/Extras/ushell/channels/unreal/core/cmds/project.py
--- a/Engine/Extras/ushell/channels/unreal/core/cmds/project.py
+++ b/Engine/Extras/ushell/channels/unreal/core/cmds/project.py
@@ -118,6 +118,7 @@ class Change(flow.cmd.Cmd):
         print("Context will now follow current directory")
         session = self.get_noticeboard(self.Noticeboard.SESSION)
         session["uproject"] = "."
+        return "."
 
     def _set_branch(self):
         self.print_info("Clearing active project")
diff --git a/Engine/Extras/ushell/channels/unreal/core/pylib/prettyprinter.py b/Engine/Extras/ushell/channels/unreal/core/pylib/prettyprinter.py
--- a/Engine/Extras/ushell/channels/unreal/core/pylib/prettyprinter.py
+++ b/Engine/Extras/ushell/channels/unreal/core/pylib/prettyprinter.py
@@ -1,9 +1,28 @@
 # Copyright Epic Games, Inc. All Rights Reserved.
 
+import sys
+
 #-------------------------------------------------------------------------------
-class Printer(object):
+class OscNineFour(object):
+    def __init__(self):
+        self._enabled = sys.stdout.isatty()
+
     def __del__(self):
-        print("\x1b]9;4;0;0\x07", end="")
+        self.clear();
+
+    def clear(self):
+        if self._enabled:
+            print("\x1b]9;4;0;0\x07", end="")
+
+    def set(self, percent):
+        if self._enabled:
+            percent = int(max(0, min(percent, 100)))
+            print(f"\x1b]9;4;1;{percent}\x07", end="")
+
+#-------------------------------------------------------------------------------
+class Printer(object):
+    def __init__(self):
+        self._osc94 = OscNineFour()
 
     def _thread(self):
         try:
@@ -27,5 +46,4 @@ class Printer(object):
         runnable.wait()
 
     def set_progress(self, percent):
-        percent = int(max(0, min(percent, 100)))
-        print(f"\x1b]9;4;1;{percent}\x07", end="")
+        self._osc94.set(percent)
diff --git a/Engine/Extras/ushell/channels/unreal/core/pylib/uelogprinter.py b/Engine/Extras/ushell/channels/unreal/core/pylib/uelogprinter.py
--- a/Engine/Extras/ushell/channels/unreal/core/pylib/uelogprinter.py
+++ b/Engine/Extras/ushell/channels/unreal/core/pylib/uelogprinter.py
@@ -7,6 +7,7 @@ import prettyprinter
 #-------------------------------------------------------------------------------
 class Printer(prettyprinter.Printer):
     def __init__(self):
+        super().__init__()
         self._log_re = re.compile(r"(\[.+\]|)Log(\w+):\s+(\w+):\s?")
         self._decorators = (
             flow.cmd.text.green,
