Updated ushell [47305291]
#rnx

diff --git a/Engine/Extras/ushell/.version b/Engine/Extras/ushell/.version
--- a/Engine/Extras/ushell/.version
+++ b/Engine/Extras/ushell/.version
@@ -1 +1 @@
-46831677
\ No newline at end of file
+47305291
\ No newline at end of file
diff --git a/Engine/Extras/ushell/channels/flow/core/cmds/zzz.py b/Engine/Extras/ushell/channels/flow/core/cmds/zzz.py
new file mode 100644
--- /dev/null
+++ b/Engine/Extras/ushell/channels/flow/core/cmds/zzz.py
@@ -0,0 +1,98 @@
+import os
+import sys
+import time
+import random
+from flow.cmd import Cmd
+
+#-------------------------------------------------------------------------------
+def one():
+    ts = os.get_terminal_size()
+    w = ts.columns - 2
+    h = ts.lines - 6
+
+    p = " .,;:*+'`^%Mx$oR?"
+
+    s0 = [[0.0] * w for x in range(h)]
+
+    for n in range(160):
+        st = time.perf_counter()
+        if n > 100:
+            s0[0] = [0.0] * w
+        else:
+            for i in range(w):
+                if random.random() < 0.1:
+                    r = random.random()
+                    s0[0][i] = 0.0 if r >= 0.5 else 0.999
+
+        for y in range(1, h):
+            for x in range(1, w - 1):
+                k  = s0[y - 1][x]
+                k += s0[y - 1][x - 1]
+                k += s0[y - 1][x + 1]
+                k += s0[y - 0][x - 1]
+                k += s0[y - 0][x + 1]
+                s0[y][x] = k / 5.01
+
+        bl = 0
+        for s in range(2, h):
+            sc = ""
+            for t in s0[h - s]:
+                i = int((t ** 1.0) * len(p))
+                bl += 1 if i > 0 else 0
+                sc += f"\x1b[9{(i % 6) + 1}m" + p[i]
+            sys.stdout.write(sc + "\n")
+        sys.stdout.write(f"\x1b[0m\x1b[{h - 2}A")
+
+        if bl == 0:
+            break
+
+        tt = time.perf_counter() - st
+        if tt < 0.03:
+            time.sleep(0.03 - tt)
+
+#-------------------------------------------------------------------------------
+def two():
+    ts = os.get_terminal_size()
+    w = ts.columns - 1
+    h = ts.lines - 5
+    ft = 0.03
+    p = "-.,:;xR+'^\"$^`*M|"
+    m = 25
+    for n in (*(x + 2 for x in range(m)), *range(m, 1, -1)):
+        st = time.perf_counter()
+        for sy in range(h):
+            s = ""
+            for sx in range(w):
+                x = ((sx * (2.67 / w)) - 2.1)
+                y = (((sy + 0.5) * (2.4 / h)) - 1.2)
+                u = v = 0
+                for i in range(n):
+                    u2, v2 = (u * u, v * v)
+                    if u2 + v2 > 3.2:
+                        break
+                    t = u2 - v2 + x
+                    v = 2 * u * v + y
+                    u = t
+                if i >= n - 1:
+                    s += " "
+                else:
+                    j = ((i * 2) % 14) + 1
+                    s += f"\x1b[9{j & 7}m" + p[(i % len(p))]
+            sys.stdout.write(s + "\n")
+        sys.stdout.write(f"\x1b[0m\x1b[{h}A")
+        if (et := ft - time.perf_counter() + st) > 0.0:
+            time.sleep(et)
+
+#-------------------------------------------------------------------------------
+class Zzz(Cmd):
+    """ Take a break """
+
+    def main(self):
+        try:
+            random.seed(time.time())
+            sys.stdout.write("\x1b[?25l")
+            (one if random.random() > 0.5 else two)()
+        except KeyboardInterrupt:
+            pass
+        finally:
+            sys.stdout.write("\x1b[?25h")
diff --git a/Engine/Extras/ushell/channels/flow/core/describe.flow.py b/Engine/Extras/ushell/channels/flow/core/describe.flow.py
--- a/Engine/Extras/ushell/channels/flow/core/describe.flow.py
+++ b/Engine/Extras/ushell/channels/flow/core/describe.flow.py
@@ -59,6 +59,10 @@ debug_wipe.source("cmds/debug.py", "Wipe")
 debug_wipe.invoke("debug", "wipe")
 debug_wipe.prefix("$")
 
+debug_zzz = flow.describe.Command()
+debug_zzz.source("cmds/zzz.py", "Zzz")
+debug_zzz.invoke("zzz")
+
 #-------------------------------------------------------------------------------
 theme_set = flow.describe.Command()
 theme_set.source("cmds/theme.py", "Set")
diff --git a/Engine/Extras/ushell/channels/unreal/core/pylib/zen/snapshot.py b/Engine/Extras/ushell/channels/unreal/core/pylib/zen/snapshot.py
--- a/Engine/Extras/ushell/channels/unreal/core/pylib/zen/snapshot.py
+++ b/Engine/Extras/ushell/channels/unreal/core/pylib/zen/snapshot.py
@@ -83,6 +83,8 @@ class UnrealCloudDDCBuildIndex(BuildIndex):
         if changelist != 0:
             query["changelist"] = {"$lte" : changelist}
 
+        # Preflights should be excluded.  There may be a usecase in the future for exposing preflights, but as of now, they are excluded
+        query["isPreflight"] = {"$eq" : False}
         query["runtime"] = {"$eq" : runtime}
         query["branch"] = {"$eq" : branch.replace("+", "/")}
 
