Updated ushell [47504908]
#rnx

diff --git a/Engine/Extras/ushell/.version b/Engine/Extras/ushell/.version
--- a/Engine/Extras/ushell/.version
+++ b/Engine/Extras/ushell/.version
@@ -1 +1 @@
-47305291
\ No newline at end of file
+47504908
\ No newline at end of file
diff --git a/Engine/Extras/ushell/README.txt b/Engine/Extras/ushell/README.txt
--- a/Engine/Extras/ushell/README.txt
+++ b/Engine/Extras/ushell/README.txt
@@ -149,6 +149,10 @@ This is required because by default ushell.bat detects if it was launched
 explicitly through Explorer (or a shortcut) and if not the assumption is it is
 running in a non-interactive scripting context (see `Scripting`)
 
+Should one find one's self faced with a choice between `Current/ushell.bat` and
+an `Active/ushell.bat`, prefer the `Current/` one. This will always exist and is
+automatically updated.
+
 # Alternative Host Shells
 
 A lot of the above assumes that cmd.exe is the host shell. Other shells are also
diff --git a/Engine/Extras/ushell/channels/unreal/core/cmds/sln.py b/Engine/Extras/ushell/channels/unreal/core/cmds/sln.py
--- a/Engine/Extras/ushell/channels/unreal/core/cmds/sln.py
+++ b/Engine/Extras/ushell/channels/unreal/core/cmds/sln.py
@@ -302,3 +302,51 @@ class Tiny(_Base):
         print("Programs:", len(programs))
 
         return self._open_sln()
+
+
+
+#-------------------------------------------------------------------------------
+class File(_Base):
+    """ Finds a suitable VS instance and opens the given file """
+    filepath    = unrealcmd.Arg(Path, "Path of the file to open")
+    line        = unrealcmd.Arg(-1, "Optional line number to open")
+    focus       = unrealcmd.Opt(False, "Bring VS into the foreground (unreliably)")
+
+    def main(self):
+        if not self.args.filepath.is_file():
+            self.print_error(f"Invalid file path '{self.args.filepath}'")
+            return False
+
+        import vs.dte
+        vs_instances = list(vs.dte.running())
+        if not vs_instances:
+            self.print_error("No VS instances appear to be running")
+            return False
+
+        target_vs = None
+        if len(vs_instances) > 1:
+            try:
+                sln_path = self._get_sln_path()
+                print("Using ", sln_path)
+            except EnvironmentError:
+                sln_path = None
+
+            for instance in vs_instances:
+                candidate = instance.get_sln_path()
+                if candidate and sln_path.samefile(candidate):
+                    target_vs = instance
+                    break
+        else:
+            target_vs = vs_instances[0]
+
+        if not target_vs:
+            self.print_error("No suitable VS instance found from a total of " + str(len(vs_instances)))
+            return False
+
+        file_path = self.args.filepath.absolute()
+        if not target_vs.open_file(str(file_path), self.args.line):
+            self.print_error("Something went wrong asking VS to open the file")
+            return False
+
+        if self.args.focus:
+            target_vs.activate()
diff --git a/Engine/Extras/ushell/channels/unreal/core/debuggers/vs.py b/Engine/Extras/ushell/channels/unreal/core/debuggers/vs.py
--- a/Engine/Extras/ushell/channels/unreal/core/debuggers/vs.py
+++ b/Engine/Extras/ushell/channels/unreal/core/debuggers/vs.py
@@ -14,6 +14,9 @@ class Debugger(object):
         cmd.launch(suspended=True, new_term=True)
         pid = cmd.get_pid()
 
+        import vs.dte
+        vs.dte.run_to_int3(pid)
+
         try:
             if not self.attach(pid):
                 cmd.kill()
diff --git a/Engine/Extras/ushell/channels/unreal/core/describe.flow.py b/Engine/Extras/ushell/channels/unreal/core/describe.flow.py
--- a/Engine/Extras/ushell/channels/unreal/core/describe.flow.py
+++ b/Engine/Extras/ushell/channels/unreal/core/describe.flow.py
@@ -205,6 +205,10 @@ sln_tiny = flow.describe.Command()
 sln_tiny.source("cmds/sln.py", "Tiny")
 sln_tiny.invoke("sln", "open", "tiny")
 
+sln_file = flow.describe.Command()
+sln_file.source("cmds/sln.py", "File")
+sln_file.invoke("sln", "open", "file")
+
 #-------------------------------------------------------------------------------
 uat = flow.describe.Command()
 uat.source("cmds/uat.py", "Uat")
diff --git a/Engine/Extras/ushell/channels/unreal/core/pylib/unrealcmd.py b/Engine/Extras/ushell/channels/unreal/core/pylib/unrealcmd.py
--- a/Engine/Extras/ushell/channels/unreal/core/pylib/unrealcmd.py
+++ b/Engine/Extras/ushell/channels/unreal/core/pylib/unrealcmd.py
@@ -44,14 +44,15 @@ class Cmd(flow.cmd.Cmd):
         platforms = self.get_unreal_context().get_platform_provider()
         yield from platforms.read_platform_names()
 
-    def get_platform(self, name=None):
+    def get_platform(self, name=None, *, safely=False):
         ue_context = self.get_unreal_context()
         platforms = ue_context.get_platform_provider()
         if platform := platforms.get_platform(name):
             return platform
 
         available = ",".join(platforms.read_platform_names())
-        raise ValueError(f"Unknown platform '{name}'. Valid list: {available}")
+        if not safely:
+            raise ValueError(f"Unknown platform '{name}'. Valid list: {available}")
 
     def get_host_platform(self):
         return self.get_platform()
diff --git a/Engine/Extras/ushell/channels/unreal/core/pylib/vs/build.bat b/Engine/Extras/ushell/channels/unreal/core/pylib/vs/build.bat
--- a/Engine/Extras/ushell/channels/unreal/core/pylib/vs/build.bat
+++ b/Engine/Extras/ushell/channels/unreal/core/pylib/vs/build.bat
@@ -10,24 +10,31 @@ if errorlevel 9009 (
 
 setlocal
 
-set vssdk=%~f1
-if "%vssdk%"=="" goto:usage
-
-set pysdk=%~f2
+set pysdk=%~f1
 if "%pysdk%"=="" goto:usage
 
-if not "%3"=="release" (
+for /f "delims=" %%d in ('vswhere -latest -find **/Inc/dte80.h') do (
+    echo found VS SDK; vssdk=%%~dpd
+    set vssdk=%%~dpd
+)
+
+if "%vssdk%"=="" (
+    echo ensure vswhere can find dte80.h by installing VS SDK.
+    goto:eof
+)
+
+if not "%2"=="debug" (
+    set _cl=/O1 /Os
+) else (
     set _cl=/Od
     set _link=/debug
-) else (
-    set _cl=/O1 /Os
 )
 
-set _cl="/I%vssdk%/Inc" "/I%pysdk%/include" /nologo /Zl /Zi /EHs-c- /c %_cl%
+set _cl="/I%vssdk%." "/I%pysdk%/include" /nologo /Zl /Zi /EHs-c- /GS- /c %_cl%
 set _link=%_link% /nologo /dll /machine:x64^
     /nodefaultlib /align:16^
     /incremental:no /opt:ref,icf^
-    "/libpath:%pysdk%/libs" python3.lib ole32.lib oleaut32.lib kernel32.lib
+    "/libpath:%pysdk%/libs" python3.lib ole32.lib oleaut32.lib kernel32.lib user32.lib
 
 mkdir 1>nul 2>nul _build
 cd _build
@@ -50,5 +57,5 @@ echo [96m%*[0m
 goto:eof
 
 :usage
-echo usage: %~nx0 vssdk pysdk [release]
+echo usage: %~nx0 pysdk [debug]
 goto:eof
diff --git a/Engine/Extras/ushell/channels/unreal/core/pylib/vs/dte.cpp b/Engine/Extras/ushell/channels/unreal/core/pylib/vs/dte.cpp
--- a/Engine/Extras/ushell/channels/unreal/core/pylib/vs/dte.cpp
+++ b/Engine/Extras/ushell/channels/unreal/core/pylib/vs/dte.cpp
@@ -65,6 +65,7 @@ public:
     void            shutdown();
     const wchar_t*  get_sln_path() const;
     const wchar_t*  get_version() const;
+    const wchar_t*  get_caption() const;
     bool            attach(int pid, const wchar_t* transport, const wchar_t* host) const;
     bool            attach(const wchar_t* name, const wchar_t* transport, const wchar_t* host) const;
     bool            open_file(const wchar_t* path, int line=-1) const;
@@ -251,6 +252,17 @@ bool Dte::open_file(const wchar_t* path, int line) const
     return true;
 }
 
+//------------------------------------------------------------------------------
+const wchar_t* Dte::get_caption() const
+{
+    ComPtr<Window> main_window;
+    if (FAILED(_dte->get_MainWindow(main_window)))
+        return nullptr;
+
+    BSTR out;
+    return main_window->get_Caption(&out) ? out : nullptr;
+}
+
 //------------------------------------------------------------------------------
 bool Dte::activate() const
 {
@@ -258,6 +270,12 @@ bool Dte::activate() const
     if (FAILED(_dte->get_MainWindow(main_window)))
         return false;
 
+    uintptr_t hwnd = 0;
+    main_window->get_HWnd((long*)&hwnd);
+    if (IsIconic(HWND(hwnd)))
+        ShowWindowAsync(HWND(hwnd), SW_RESTORE);
+    SetForegroundWindow(HWND(hwnd));
+
     return SUCCEEDED(main_window->Activate());
 }
 
@@ -564,6 +582,19 @@ static void suspend_resume_init()
     NtSuspendProcess = decltype(NtSuspendProcess)(GetProcAddress(ntdll, "NtSuspendProcess"));
 }
 
+//------------------------------------------------------------------------------
+static int suspend_resume(bool is_suspend, unsigned pid)
+{
+    HANDLE handle = OpenProcess(PROCESS_SUSPEND_RESUME, FALSE, pid);
+    if (handle == nullptr)
+        return -1;
+
+    int ret = is_suspend ? NtSuspendProcess(handle) : NtResumeProcess(handle);
+
+    CloseHandle(handle);
+    return ret;
+}
+
 //------------------------------------------------------------------------------
 static PyObject* suspend_resume(bool is_suspend, PyObject** args, int nargs)
 {
@@ -575,13 +606,7 @@ static PyObject* suspend_resume(bool is_suspend, PyObject** args, int nargs)
         Py_RETURN_FALSE;
 
     unsigned pid = PyLong_AsLong(py_pid);
-
-    HANDLE handle = OpenProcess(PROCESS_SUSPEND_RESUME, FALSE, pid);
-    if (handle == nullptr)
-        Py_RETURN_FALSE;
-
-    int ok = is_suspend ? NtSuspendProcess(handle) : NtResumeProcess(handle);
-    CloseHandle(handle);
+    int ok = suspend_resume(is_suspend, pid);
 
     if (ok < 0)
         Py_RETURN_FALSE;
@@ -601,12 +626,62 @@ static PyObject* resume_process(PyObject* self, PyObject** args, Py_ssize_t narg
     return suspend_resume(false, args, int(nargs));
 }
 
+//------------------------------------------------------------------------------
+static PyObject* run_to_int3(PyObject* self, PyObject** args, Py_ssize_t nargs)
+{
+    static int iii = 0;
+    while (!iii)
+        Sleep(100);
+
+    if (nargs < 1)
+        Py_RETURN_FALSE;
+
+    PyObject* py_pid = args[0];
+    if (!PyLong_Check(py_pid))
+        Py_RETURN_FALSE;
+
+    unsigned pid = PyLong_AsLong(py_pid);
+    if (DebugActiveProcess(pid) == FALSE)
+        Py_RETURN_FALSE;
+
+    suspend_resume(false, pid);
+
+    int count = 2;
+    for (
+        DEBUG_EVENT event;
+        WaitForDebugEvent(&event, 0x49e) == TRUE;
+        ContinueDebugEvent(event.dwProcessId, event.dwThreadId, DBG_CONTINUE))
+    {
+        if (event.dwDebugEventCode != EXCEPTION_DEBUG_EVENT)
+            continue;
+
+        if (event.u.Exception.ExceptionRecord.ExceptionCode != EXCEPTION_BREAKPOINT)
+            continue;
+
+        // There are two breakpoints, one from the attach and one from the loader.
+        if (--count > 0)
+            continue;
+
+        suspend_resume(true, pid);
+        break;
+    }
+
+    DebugSetProcessKillOnExit(FALSE);
+    DebugActiveProcessStop(pid);
+
+    if (count != 0)
+        Py_RETURN_FALSE;
+
+    Py_RETURN_TRUE;
+}
+
 
 
 //------------------------------------------------------------------------------
 static PyMethodDef dte_module_methods[] = {
     { "suspend_process", PyCFunction(suspend_process), METH_FASTCALL, "Suspend a process by PID"},
     { "resume_process",  PyCFunction(resume_process),  METH_FASTCALL, "Resume a process by PID"},
+    { "run_to_int3",     PyCFunction(run_to_int3),     METH_FASTCALL, "Run until NT's implicit int3s then suspends"},
     {},
 };
 
diff --git a/Engine/Extras/ushell/channels/unreal/core/pylib/zen/snapshot.py b/Engine/Extras/ushell/channels/unreal/core/pylib/zen/snapshot.py
--- a/Engine/Extras/ushell/channels/unreal/core/pylib/zen/snapshot.py
+++ b/Engine/Extras/ushell/channels/unreal/core/pylib/zen/snapshot.py
@@ -330,8 +330,8 @@ class _Base(zen.cmd.ZenUtilityBaseCmd):
     def get_cook_flavor(self):
         effective_flavor = self.args.flavor
         if not effective_flavor:
-            platform = self.get_platform(self.args.platform)
-            effective_flavor = platform.get_cook_flavor(self.args.runtime.lower())
+            if platform := self.get_platform(self.args.platform, safely=True):
+                effective_flavor = platform.get_cook_flavor(self.args.runtime.lower())
         return effective_flavor
 
     def get_build_index(self):
