Updated ushell [47508336]
#rnx

diff --git a/Engine/Extras/ushell/.version b/Engine/Extras/ushell/.version
--- a/Engine/Extras/ushell/.version
+++ b/Engine/Extras/ushell/.version
@@ -1 +1 @@
-47504908
\ No newline at end of file
+47508336
\ No newline at end of file
diff --git a/Engine/Extras/ushell/channels/unreal/core/pylib/vs/dte.cpp b/Engine/Extras/ushell/channels/unreal/core/pylib/vs/dte.cpp
--- a/Engine/Extras/ushell/channels/unreal/core/pylib/vs/dte.cpp
+++ b/Engine/Extras/ushell/channels/unreal/core/pylib/vs/dte.cpp
@@ -629,10 +629,6 @@ static PyObject* resume_process(PyObject* self, PyObject** args, Py_ssize_t narg
 //------------------------------------------------------------------------------
 static PyObject* run_to_int3(PyObject* self, PyObject** args, Py_ssize_t nargs)
 {
-    static int iii = 0;
-    while (!iii)
-        Sleep(100);
-
     if (nargs < 1)
         Py_RETURN_FALSE;
 
