Updated ushell [48218701]
#rnx

diff --git a/Engine/Extras/ushell/.version b/Engine/Extras/ushell/.version
--- a/Engine/Extras/ushell/.version
+++ b/Engine/Extras/ushell/.version
@@ -1 +1 @@
-47508336
\ No newline at end of file
+48218701
\ No newline at end of file
diff --git a/Engine/Extras/ushell/channels/unreal/core/cmds/sln.py b/Engine/Extras/ushell/channels/unreal/core/cmds/sln.py
--- a/Engine/Extras/ushell/channels/unreal/core/cmds/sln.py
+++ b/Engine/Extras/ushell/channels/unreal/core/cmds/sln.py
@@ -3,6 +3,8 @@
 import os
 import shutil
 import unrealcmd
+import slnutils
+import subprocess
 from pathlib import Path
 
 #-------------------------------------------------------------------------------
@@ -350,3 +352,55 @@ class File(_Base):
 
         if self.args.focus:
             target_vs.activate()
+
+
+#-------------------------------------------------------------------------------
+class OpenVS(_Base):
+    """ Opens a solution in Visual Studio. """
+
+    def main(self):
+        if os.name != "nt":
+            self.print_error("Opening a solution in Visual Studio is only supported on Windows")
+            return False
+
+        self.print_info("Opening solution in Visual Studio")
+        sln_path = self._get_sln_path()
+        print("Path:", sln_path)
+        if not os.path.isfile(sln_path):
+            self.print_error("Project file not found")
+            return False
+
+        args = (
+            sln_path,
+        )
+        
+        run_args = ("cmd.exe", "/d", "/c", "start devenv", sln_path)
+        shell_open = self.get_exec_context().create_runnable(*run_args)
+        return shell_open.run()
+
+
+#-------------------------------------------------------------------------------
+class OpenRider(_Base):
+    """ Opens a solution in Rider. """
+
+    def main(self):
+        if os.name != "nt":
+            self.print_error("Opening a solution in Rider is only supported on Windows")
+            return False
+
+        self.print_info("Opening solution in Rider")
+        sln_path = self._get_sln_path()
+        print("Path:", sln_path)
+        if not os.path.isfile(sln_path):
+            self.print_error("Project file not found")
+            return False
+
+        run_args = ("cmd.exe", "/d", "/c", slnutils.find_rider_bin(), sln_path)
+
+        try:
+            print("Launching Rider:", *run_args)
+            proc = subprocess.Popen(run_args)
+            return True
+        except Exception as e:
+            print("!! Failed launching Rider:", str(e))
+            return False
\ No newline at end of file
diff --git a/Engine/Extras/ushell/channels/unreal/core/debuggers/rider.py b/Engine/Extras/ushell/channels/unreal/core/debuggers/rider.py
--- a/Engine/Extras/ushell/channels/unreal/core/debuggers/rider.py
+++ b/Engine/Extras/ushell/channels/unreal/core/debuggers/rider.py
@@ -2,6 +2,7 @@
 
 import os
 import subprocess
+import slnutils
 from pathlib import Path
 
 #-------------------------------------------------------------------------------
@@ -31,22 +32,6 @@ def _select_sln_path(ue_context):
 
     return fallback
 
-#-------------------------------------------------------------------------------
-def _find_rider_bin():
-    if os.name != "nt":
-        return Path("rider64")
-
-    root_dir = Path(os.getenv("ProgramFiles", "c:/Program Files/"))
-    root_dir /= "JetBrains"
-    candidates = [x for x in root_dir.glob("JetBrains Rider*") if x.is_dir()]
-    if candidates:
-        candidates.sort(key=lambda x: x.stat().st_mtime)
-        candidate = candidates[-1] / "bin/rider64.exe"
-        if candidate.is_file():
-            return candidate
-
-    return Path("rider64.exe")
-
 #-------------------------------------------------------------------------------
 class Debugger(object):
     def __init__(self, ue_context):
@@ -69,7 +54,7 @@ class Debugger(object):
             print("Exotic debugger attachment isn't supported!")
             return False
 
-        bin_path = _find_rider_bin()
+        bin_path = slnutils.find_rider_bin()
         rider_args = (bin_path, "attach-to-process", str(pid))
         if sln_path := _select_sln_path(self._ue_context):
             sln_path = sln_path.absolute()
diff --git a/Engine/Extras/ushell/channels/unreal/core/describe.flow.py b/Engine/Extras/ushell/channels/unreal/core/describe.flow.py
--- a/Engine/Extras/ushell/channels/unreal/core/describe.flow.py
+++ b/Engine/Extras/ushell/channels/unreal/core/describe.flow.py
@@ -209,6 +209,14 @@ sln_file = flow.describe.Command()
 sln_file.source("cmds/sln.py", "File")
 sln_file.invoke("sln", "open", "file")
 
+sln_vs = flow.describe.Command()
+sln_vs.source("cmds/sln.py", "OpenVS")
+sln_vs.invoke("sln", "open", "vs")
+
+sln_rider = flow.describe.Command()
+sln_rider.source("cmds/sln.py", "OpenRider")
+sln_rider.invoke("sln", "open", "rider")
+
 #-------------------------------------------------------------------------------
 uat = flow.describe.Command()
 uat.source("cmds/uat.py", "Uat")
diff --git a/Engine/Extras/ushell/channels/unreal/core/pylib/slnutils.py b/Engine/Extras/ushell/channels/unreal/core/pylib/slnutils.py
new file mode 100644
--- /dev/null
+++ b/Engine/Extras/ushell/channels/unreal/core/pylib/slnutils.py
@@ -0,0 +1,20 @@
+# Copyright Epic Games, Inc. All Rights Reserved.
+
+import os
+from pathlib import Path
+
+#-------------------------------------------------------------------------------
+def find_rider_bin():
+    if os.name != "nt":
+        return Path("rider64")
+
+    root_dir = Path(os.getenv("ProgramFiles", "c:/Program Files/"))
+    root_dir /= "JetBrains"
+    candidates = [x for x in root_dir.glob("JetBrains Rider*") if x.is_dir()]
+    if candidates:
+        candidates.sort(key=lambda x: x.stat().st_mtime)
+        candidate = candidates[-1] / "bin/rider64.exe"
+        if candidate.is_file():
+            return candidate
+
+    return Path("rider64.exe")
