Updated ushell [48360749]
#rnx

diff --git a/Engine/Extras/ushell/.version b/Engine/Extras/ushell/.version
--- a/Engine/Extras/ushell/.version
+++ b/Engine/Extras/ushell/.version
@@ -1 +1 @@
-48218701
\ No newline at end of file
+48360749
\ No newline at end of file
diff --git a/Engine/Extras/ushell/channels/unreal/core/cmds/sln.py b/Engine/Extras/ushell/channels/unreal/core/cmds/sln.py
--- a/Engine/Extras/ushell/channels/unreal/core/cmds/sln.py
+++ b/Engine/Extras/ushell/channels/unreal/core/cmds/sln.py
@@ -2,6 +2,7 @@
 
 import os
 import shutil
+import unreal
 import unrealcmd
 import slnutils
 import subprocess
@@ -211,7 +212,65 @@ class Open10x(_Base):
         result = cmd.launch()
         return True
 
+#-------------------------------------------------------------------------------
+class OpenVSCode(unrealcmd.MultiPlatformCmd):
+    """ Opens a Visual Studio Code workspace """
 
+    def _get_wokspace_path(self):
+        ue_context = self.get_unreal_context()
+        engine = ue_context.get_engine()
+        engine_dir_name = engine.get_dir().parent.stem
+
+        primary_name = None
+        try:
+            config = engine.get_ubt().get_configuration()
+            config.get_value("ProjectFileGenerator", "PrimaryProjectName")
+        except Exception:
+            primary_name = "UE5"
+
+        possible_names = []
+        if project := ue_context.get_project():
+            project_name = project.get_name()
+            # could be name of the project itself
+            possible_names.append(project_name)
+            # if bPrimaryProjectNameFromFolder - could append engine directory name
+            possible_names.append(f"{project_name}_{engine_dir_name}")
+
+        # could be just "UE5"
+        possible_names.append(primary_name)
+        # if UE_NAME_PROJECT_AFTER_FOLDER:
+        possible_names.append(f"{primary_name}_{engine_dir_name}")
+        # if bPrimaryProjectNameFromFolder:
+        possible_names.append(engine_dir_name)
+
+        workspace_dir = engine.get_dir().parent
+        if ue_context.get_type() == unreal.ContextType.FOREIGN:
+            workspace_dir = ue_context.get_project().get_dir()
+
+        for workspace_name in possible_names:
+            workspace_path = workspace_dir / f"{workspace_name}.code-workspace"
+            if workspace_path.is_file():
+                return workspace_path
+        # no project is generated - open as a folder
+        return workspace_dir
+
+    def main(self):
+        self.print_info("Opening workspace in VSCode")
+        workspace_path = self._get_wokspace_path()
+        if not os.path.exists(workspace_path):
+            self.print_error("Workspace not found")
+            return False
+        print("Path:", workspace_path)
+
+        code_executable = "code.cmd" if os.name == "nt" else "code"
+
+        args = (
+            workspace_path,
+        )
+
+        cmd = self.get_exec_context().create_runnable(code_executable, *args)
+        result = cmd.launch()
+        return True
 
 #-------------------------------------------------------------------------------
 class Tiny(_Base):
diff --git a/Engine/Extras/ushell/channels/unreal/core/describe.flow.py b/Engine/Extras/ushell/channels/unreal/core/describe.flow.py
--- a/Engine/Extras/ushell/channels/unreal/core/describe.flow.py
+++ b/Engine/Extras/ushell/channels/unreal/core/describe.flow.py
@@ -201,6 +201,10 @@ sln_10x = flow.describe.Command()
 sln_10x.source("cmds/sln.py", "Open10x")
 sln_10x.invoke("sln", "open", "10x")
 
+sln_vscode = flow.describe.Command()
+sln_vscode.source("cmds/sln.py", "OpenVSCode")
+sln_vscode.invoke("sln", "open", "code")
+
 sln_tiny = flow.describe.Command()
 sln_tiny.source("cmds/sln.py", "Tiny")
 sln_tiny.invoke("sln", "open", "tiny")
