Updated ushell [49165608]
#rnx

diff --git a/Engine/Extras/ushell/.version b/Engine/Extras/ushell/.version
--- a/Engine/Extras/ushell/.version
+++ b/Engine/Extras/ushell/.version
@@ -1 +1 @@
-48360749
\ No newline at end of file
+49165608
\ No newline at end of file
diff --git a/Engine/Extras/ushell/channels/unreal/core/cmds/getbuild.py b/Engine/Extras/ushell/channels/unreal/core/cmds/getbuild.py
--- a/Engine/Extras/ushell/channels/unreal/core/cmds/getbuild.py
+++ b/Engine/Extras/ushell/channels/unreal/core/cmds/getbuild.py
@@ -9,6 +9,18 @@ import os
 from peafour import P4
 import string
 import fzf
+import subprocess
+from flow.text import text
+
+def run_cmd(cmd, dryrun = False):
+    with text.cyan:
+        print("== Running command: ", end="")
+    print("." + cmd[1:])
+    if not dryrun:
+        ret = subprocess.run(cmd)
+        if ret.returncode:
+            return (False, ret.returncode)
+    return (True, "")
 
 #-------------------------------------------------------------------------------
 class GetBuild(unrealcmd.Cmd):
@@ -35,6 +47,8 @@ If you find multiple builds that match you will be given the option to pick one,
     verbose             = unrealcmd.Opt(False, "Enable verbose logging output")
     preflight           = unrealcmd.Opt(False, "If set, shows preflights")
     wildcard            = unrealcmd.Opt("", "Windows style wildcard (using * and ?) to match file paths to include")
+    sync                = unrealcmd.Opt(False, "Use P4 to sync to the specified changelist.")
+    buildgroup          = unrealcmd.Opt("", "The build group to filter, this is frequently globally unique and thus allows you to uniquely specify a build without any other options")
 
     dest                = unrealcmd.Opt("", "The directory to download the build to. Local contents will be overwritten. Defaults to `<projectdir>/Saved/<StagedBuilds/PackagedBuilds>/<platform>`")
     match               = unrealcmd.Opt("", "Either a build id (hexadecimal string), a changelist (number) or a key value separate list of key=values you want to find a build for. e.g. changelist=1337;job=foo")
@@ -119,10 +133,14 @@ If you find multiple builds that match you will be given the option to pick one,
                 yield item.name + "/"
 
     def convert_platform(self, platform):
-        was_server = False
-        if platform.endswith("-server"):
-            platform = platform[0:-7]
-            was_server = True
+        # strip suffix to map platform name separately
+        def strip_suffix(platform):
+            for suffix in ["-server", "-client"]:
+                if platform.endswith(suffix):
+                    return platform[0:-len(suffix)], suffix
+            return platform, ""
+
+        platform, suffix = strip_suffix(platform)
 
         # non-windows platforms do not load the windows platform but we still need to translate the platform name
         # as such we hardcode the conversion for this platform
@@ -131,12 +149,10 @@ If you find multiple builds that match you will be given the option to pick one,
 
         try:
             converted_platform = self.get_platform(platform).get_config_name()
-            if was_server:
-                return converted_platform + "-server"
-            return converted_platform
+            return converted_platform + suffix
         except ValueError:
             # unknown platform, just assume the name it valid
-            return platform
+            return platform + suffix
 
     def convert_buildtype(self, build_type):
         if build_type == "packaged":
@@ -178,7 +194,7 @@ If you find multiple builds that match you will be given the option to pick one,
             return build_id_arg
         return None
 
-    def build_query(self, search_request, build_type, branch, platforms, preflight):
+    def build_query(self, search_request, build_type, branch, platforms, preflight, buildgroup):
         search_changelist = None
         if search_request.isnumeric():
             search_changelist = search_request
@@ -190,13 +206,17 @@ If you find multiple builds that match you will be given the option to pick one,
             search_request = "" # special case handling the build_id filtering
 
         query = {}
+
         for statement in search_request.split(";"):
-            if statement.find("=") != -1:
+            if statement.find("!=") != -1:
+                k, v = statement.split("!=")
+                query[k] = {"$neq": v}
+            elif statement.find("=") != -1:
                 k, v = statement.split("=")
                 query[k] = {"$eq": v}
 
         if search_changelist:
-            query["commit"] = {"$lte": int(search_changelist)}
+            query["commit"] = {"$eq": int(search_changelist)}
 
         if build_id:
             query["buildId"] = {"$eq": str(build_id)}
@@ -204,7 +224,12 @@ If you find multiple builds that match you will be given the option to pick one,
         if build_type:
             query["type"] = {"$eq": str(build_type)}
 
-        if branch:
+        if buildgroup:
+            query["buildgroup"] = {"$eq": str(buildgroup)}
+
+        # ignore the defaulted branch filtering when specifying build group as that uniquely identifies the branch anyway
+        # this allows us to download a build from any branch if you know its build group
+        if branch and not buildgroup:
             query["branch"] = {"$eq": str(branch.lower())}
 
         if not preflight:
@@ -216,10 +241,11 @@ If you find multiple builds that match you will be given the option to pick one,
             for platform in platforms:
                 p = platform.lower()
                 platform_filters.append(p)
-                if not p.endswith("-server"):
+                # if not explicitly specified, add -client suffix for every platform
+                if not p.endswith("-server") and not p.endswith("-client"):
                     platform_filters.append(f"{p}-client")
-
-                if self.is_server_platform(p):
+                # if not explicitly specified, add -server suffix for server platform
+                if not p.endswith("-server") and self.is_server_platform(p):
                     platform_filters.append(f"{p}-server")
 
             query["platform"] =  {"$in": platform_filters}
@@ -385,7 +411,7 @@ If you find multiple builds that match you will be given the option to pick one,
         builds_to_download = []
         found_builds = {}
 
-        query = self.build_query(self.args.match, build_type, branch, platforms_to_download, self.args.preflight)
+        query = self.build_query(self.args.match, build_type, branch, platforms_to_download, self.args.preflight, self.args.buildgroup)
         bucket_regex = f"{project}.{build_type}.*".lower()
         request = {
             "bucketRegex": bucket_regex,
@@ -459,6 +485,12 @@ If you find multiple builds that match you will be given the option to pick one,
                 "--exclude-wildcard=" + getattr(self.args, "exclude-wildcard", "")
             )
 
-            result = self.run_zen(args)
+            result |= self.run_zen(args)
+
+        if result == 0 and self.args.sync:
+            (result, err) = run_cmd(f"_p4 sync {selected_build["commit"]}")
+            if result == False:
+                self.print_error(f"P4 sync failed with error: {err}. Aborting.")
+                return False
 
         return result
diff --git a/Engine/Extras/ushell/channels/unreal/core/cmds/sln.py b/Engine/Extras/ushell/channels/unreal/core/cmds/sln.py
--- a/Engine/Extras/ushell/channels/unreal/core/cmds/sln.py
+++ b/Engine/Extras/ushell/channels/unreal/core/cmds/sln.py
@@ -8,6 +8,9 @@ import slnutils
 import subprocess
 from pathlib import Path
 
+UPROJ_OPT_STR = "Prefer opening active project's .uproject ('--uproj=no' to disable,"
+" 'USHELL_UPROJ=yes/no' environment variable to change the default behaviour)."
+
 #-------------------------------------------------------------------------------
 def _make_temp_vcxproj(ctx, path, name, target, ubt_args=""):
     eng_dir = ctx.get_engine().get_dir()
@@ -73,6 +76,12 @@ class _Base(unrealcmd.MultiPlatformCmd):
 
         return sln_path
 
+    def _get_uproject_path(self):
+        ue_context = self.get_unreal_context()
+
+        if project := ue_context.get_project():
+            return project.get_path()
+
     def _open_sln(self):
         if os.name != "nt":
             self.print_error("Opening project files is only supported on Windows")
@@ -416,6 +425,7 @@ class File(_Base):
 #-------------------------------------------------------------------------------
 class OpenVS(_Base):
     """ Opens a solution in Visual Studio. """
+    uproj = unrealcmd.Opt((None, "yes", "no"), UPROJ_OPT_STR)
 
     def main(self):
         if os.name != "nt":
@@ -423,16 +433,31 @@ class OpenVS(_Base):
             return False
 
         self.print_info("Opening solution in Visual Studio")
-        sln_path = self._get_sln_path()
+
+        use_uproj = slnutils.use_uproj(self.args)
+
+        if use_uproj:
+            sln_path = self._get_uproject_path()
+            if sln_path:
+                sln_path = sln_path.parent
+        else:
+            sln_path = self._get_sln_path()
+
         print("Path:", sln_path)
-        if not os.path.isfile(sln_path):
-            self.print_error("Project file not found")
-            return False
+
+        if use_uproj:
+            if not sln_path:
+                self.print_error("Project path not found")
+                return False
+        else:
+            if not os.path.isfile(sln_path):
+                self.print_error("Project file not found")
+                return False
 
         args = (
             sln_path,
         )
-        
+
         run_args = ("cmd.exe", "/d", "/c", "start devenv", sln_path)
         shell_open = self.get_exec_context().create_runnable(*run_args)
         return shell_open.run()
@@ -441,6 +466,7 @@ class OpenVS(_Base):
 #-------------------------------------------------------------------------------
 class OpenRider(_Base):
     """ Opens a solution in Rider. """
+    uproj = unrealcmd.Opt((None, "yes", "no"), UPROJ_OPT_STR)
 
     def main(self):
         if os.name != "nt":
@@ -448,7 +474,14 @@ class OpenRider(_Base):
             return False
 
         self.print_info("Opening solution in Rider")
-        sln_path = self._get_sln_path()
+
+        use_uproj = slnutils.use_uproj(self.args)
+
+        if use_uproj:
+            sln_path = self._get_uproject_path()
+        else:
+            sln_path = self._get_sln_path()
+
         print("Path:", sln_path)
         if not os.path.isfile(sln_path):
             self.print_error("Project file not found")
diff --git a/Engine/Extras/ushell/channels/unreal/core/debuggers/vs.py b/Engine/Extras/ushell/channels/unreal/core/debuggers/vs.py
--- a/Engine/Extras/ushell/channels/unreal/core/debuggers/vs.py
+++ b/Engine/Extras/ushell/channels/unreal/core/debuggers/vs.py
@@ -31,7 +31,7 @@ class Debugger(object):
         import vs.dte
         for instance in vs.dte.running():
             candidate = instance.get_sln_path()
-            if candidate and engine_dir.is_relative_to(Path(candidate).parent):
+            if candidate and Path(candidate).is_relative_to(engine_dir.parent):
                 if instance.attach(pid, transport, host_ip):
                     break
         else:
diff --git a/Engine/Extras/ushell/channels/unreal/core/pylib/slnutils.py b/Engine/Extras/ushell/channels/unreal/core/pylib/slnutils.py
--- a/Engine/Extras/ushell/channels/unreal/core/pylib/slnutils.py
+++ b/Engine/Extras/ushell/channels/unreal/core/pylib/slnutils.py
@@ -18,3 +18,9 @@ def find_rider_bin():
             return candidate
 
     return Path("rider64.exe")
+
+def use_uproj(args):
+    if args.uproj is None:
+        return os.getenv("USHELL_UPROJ", "no").lower() == "yes"
+    else:
+        return args.uproj.lower() == "yes"
