Updated ushell [49372623]
#rnx

diff --git a/Engine/Extras/ushell/.version b/Engine/Extras/ushell/.version
--- a/Engine/Extras/ushell/.version
+++ b/Engine/Extras/ushell/.version
@@ -1 +1 @@
-49165608
\ No newline at end of file
+49372623
\ No newline at end of file
diff --git a/Engine/Extras/ushell/channels/unreal/core/cmds/getbuild.py b/Engine/Extras/ushell/channels/unreal/core/cmds/getbuild.py
--- a/Engine/Extras/ushell/channels/unreal/core/cmds/getbuild.py
+++ b/Engine/Extras/ushell/channels/unreal/core/cmds/getbuild.py
@@ -49,6 +49,7 @@ If you find multiple builds that match you will be given the option to pick one,
     wildcard            = unrealcmd.Opt("", "Windows style wildcard (using * and ?) to match file paths to include")
     sync                = unrealcmd.Opt(False, "Use P4 to sync to the specified changelist.")
     buildgroup          = unrealcmd.Opt("", "The build group to filter, this is frequently globally unique and thus allows you to uniquely specify a build without any other options")
+    config              = unrealcmd.Opt("", "The Configuration to filter")
 
     dest                = unrealcmd.Opt("", "The directory to download the build to. Local contents will be overwritten. Defaults to `<projectdir>/Saved/<StagedBuilds/PackagedBuilds>/<platform>`")
     match               = unrealcmd.Opt("", "Either a build id (hexadecimal string), a changelist (number) or a key value separate list of key=values you want to find a build for. e.g. changelist=1337;job=foo")
@@ -194,7 +195,7 @@ If you find multiple builds that match you will be given the option to pick one,
             return build_id_arg
         return None
 
-    def build_query(self, search_request, build_type, branch, platforms, preflight, buildgroup):
+    def build_query(self, search_request, build_type, branch, platforms, preflight, buildgroup, config):
         search_changelist = None
         if search_request.isnumeric():
             search_changelist = search_request
@@ -227,6 +228,9 @@ If you find multiple builds that match you will be given the option to pick one,
         if buildgroup:
             query["buildgroup"] = {"$eq": str(buildgroup)}
 
+        if config:
+            query["Configuration"] = {"$eq": str(config)}
+
         # ignore the defaulted branch filtering when specifying build group as that uniquely identifies the branch anyway
         # this allows us to download a build from any branch if you know its build group
         if branch and not buildgroup:
@@ -283,7 +287,8 @@ If you find multiple builds that match you will be given the option to pick one,
             "platforms": {r["metadata"].get("cookPlatform", r["metadata"]["platform"]): (r["buildId"], r["bucketId"])}, # use cook platform over platform if present as that field follows UE conventions for how the platform is called
             "buildgroup" : r["metadata"].get("buildgroup", ""),
             "commit": r["metadata"].get("commit", ""),
-            "hordetemplate" : r["metadata"].get("hordetemplateid", "")
+            "hordetemplate" : r["metadata"].get("hordetemplateid", ""),
+            "config": r["metadata"].get("Configuration", "")
         } for r in results]
 
     def _show_select_build_prompt(self, branch_cl, builds):
@@ -321,6 +326,7 @@ If you find multiple builds that match you will be given the option to pick one,
                 name = build.get("buildgroup", "")
                 if not name:
                     name = build["name"]
+                config = build.get("config", "")
                 platforms = build["platforms"].keys()
                 middle = ""
                 for column_platform,short_name in columns.items():
@@ -331,13 +337,14 @@ If you find multiple builds that match you will be given the option to pick one,
                     if not p4_change_found:
                         closest_cl = branch_cl
                     cl = str(cl) + ("*" if cl == closest_cl else " ")
-                yield "%10s%s %s" % (cl, middle, name)
+                yield "%10s%s %s %s" % (cl, middle, name, config)
 
         for line in fzf.run(list_available(closest_cl, available_builds), sort=False):
             segments = line.split()
-            build_name = segments[-1]
+            build_name = segments[-2]
+            config_name = segments[-1]
             commit = segments[0].rstrip("*")
-            build_key = build_name + commit
+            build_key = build_name + commit + config_name
 
             if build_key in builds:
                 return build_key
@@ -411,7 +418,7 @@ If you find multiple builds that match you will be given the option to pick one,
         builds_to_download = []
         found_builds = {}
 
-        query = self.build_query(self.args.match, build_type, branch, platforms_to_download, self.args.preflight, self.args.buildgroup)
+        query = self.build_query(self.args.match, build_type, branch, platforms_to_download, self.args.preflight, self.args.buildgroup, self.args.config)
         bucket_regex = f"{project}.{build_type}.*".lower()
         request = {
             "bucketRegex": bucket_regex,
@@ -425,7 +432,8 @@ If you find multiple builds that match you will be given the option to pick one,
             if not build_name:
                 build_name = build["name"]
             commit_str = str(build.get("commit", ""))
-            build_key = build_name + commit_str
+            config_str = str(build.get("config", ""))
+            build_key = build_name + commit_str + config_str
 
             if build_key in found_builds:
                 # this build already existed but a new platform was found add it to the list if it didn't already exist
diff --git a/Engine/Extras/ushell/channels/unreal/core/cmds/perftest.py b/Engine/Extras/ushell/channels/unreal/core/cmds/perftest.py
--- a/Engine/Extras/ushell/channels/unreal/core/cmds/perftest.py
+++ b/Engine/Extras/ushell/channels/unreal/core/cmds/perftest.py
@@ -321,6 +321,23 @@ class PerfTestBaseCmd(unrealcmd.Cmd):
         config_key      = "MaterialsToTest"
         return self._find_config(config_section, config_key)
 
+    def _get_asset_name(self, object_path) -> str:
+        # Remove spaces and quotes from the start and end.
+        s = object_path.strip(' \'"')
+        if not s:
+            return ''
+
+        # Remove any sub-objects by splitting on the first :.
+        s = s.split(':', 1)[0]
+
+        # Get the last segment
+        basename = s.rsplit('/', 1)[-1]
+
+        # Use the asset name rather than package name.
+        if '.' in basename:
+            basename = basename.rsplit('.', 1)[1]
+        return basename
+
     def _make_sequence_args(self, sequence:str) -> tuple[str]:
         return (f"-AutomatedPerfTest.SequencePerfTest.MapSequenceName={sequence}",)
 
@@ -345,7 +362,7 @@ class Sequence(PerfTestBaseCmd):
     Sequence Test - Runs a configured sequence in the project while collecting
     performance metrics.
     """
-    SequenceComboName   = unrealcmd.Arg(str,  "Sequence Combo Name")
+    SequenceComboName   = unrealcmd.Arg(str, "Sequence Combo Name")
 
     def main(self):
         args = self._make_sequence_args(self.args.SequenceComboName)
@@ -400,14 +417,15 @@ class StaticCamera(PerfTestBaseCmd):
     """
     Static Camera Test.
     """
-    MapName  = unrealcmd.Arg(str,  "Static Camera Test Map Name")
+    MapName  = unrealcmd.Opt("",  "Static Camera Test Map Name")
     def main(self):
-        args = self._make_static_cam_args(self.args.MapName)
+        args = (None,)
+        if self.args.MapName != "":
+           args = self._make_static_cam_args(self.args.MapName)
         return self.run_perftest(PerfTestType.CAMERA, *args)
 
     def complete_MapName(self, prefix):
-        yield from self._get_available_static_cameratests()
-
+        return [self._get_asset_name(x) for x in self._get_available_static_cameratests()]
 
 #-------------------------------------------------------------------------------
 class PerfTestDefault(PerfTestBaseCmd):
