Updated ushell [49945479]
#rnx

diff --git a/Engine/Extras/ushell/.version b/Engine/Extras/ushell/.version
--- a/Engine/Extras/ushell/.version
+++ b/Engine/Extras/ushell/.version
@@ -1 +1 @@
-49372623
\ No newline at end of file
+49945479
\ No newline at end of file
diff --git a/Engine/Extras/ushell/channels/unreal/core/cmds/cook.py b/Engine/Extras/ushell/channels/unreal/core/cmds/cook.py
--- a/Engine/Extras/ushell/channels/unreal/core/cmds/cook.py
+++ b/Engine/Extras/ushell/channels/unreal/core/cmds/cook.py
@@ -98,10 +98,15 @@ class _Cook(unrealcmd.MultiPlatformCmd):
         variant = unreal.Variant.DEBUG if self.args.debug else unreal.Variant.DEVELOPMENT
         build = target.get_build(variant=variant)
         if not build:
-            self.print_error(f"No {variant.name.lower()} editor build found")
+            self.print_error(f"No {variant.name.lower()} editor build found (check .target receipt in Binaries/)")
             return False
 
-        cmd = str(build.get_binary_path())
+        cmd = build.get_binary_path()
+        if cmd is None:
+            self.print_error(f"Unable to get 'Launch' path from .target receipt in Binaries/")
+            return False
+
+        cmd = str(cmd)
         if not cmd.endswith("-Cmd.exe"):
             cmd = cmd.replace(".exe", "-Cmd.exe")
 
diff --git a/Engine/Extras/ushell/channels/unreal/core/cmds/sln.py b/Engine/Extras/ushell/channels/unreal/core/cmds/sln.py
--- a/Engine/Extras/ushell/channels/unreal/core/cmds/sln.py
+++ b/Engine/Extras/ushell/channels/unreal/core/cmds/sln.py
@@ -6,6 +6,7 @@ import unreal
 import unrealcmd
 import slnutils
 import subprocess
+import sys
 from pathlib import Path
 
 UPROJ_OPT_STR = "Prefer opening active project's .uproject ('--uproj=no' to disable,"
@@ -63,18 +64,24 @@ class _Base(unrealcmd.MultiPlatformCmd):
 
     def _get_sln_path(self):
         ue_context = self.get_unreal_context()
-
         primary_name = self._get_primary_name()
-        base_name = primary_name + ".sln"
 
-        sln_path = None
+        candidates = []
+        def gather(candidate_dir):
+            for ext in (".slnx", ".sln"):
+                candidates.append(candidate_dir / (primary_name + ext))
+
         if project := ue_context.get_project():
-            sln_path = project.get_dir() / base_name
-        if sln_path is None or not sln_path.is_file():
-            engine = ue_context.get_engine()
-            sln_path = engine.get_dir().parent / base_name
+            gather(project.get_dir())
 
-        return sln_path
+        engine = ue_context.get_engine()
+        gather(engine.get_dir().parent)
+
+        for candidate in candidates:
+            if candidate.is_file():
+                return candidate
+
+        return candidates[-1]
 
     def _get_uproject_path(self):
         ue_context = self.get_unreal_context()
@@ -495,4 +502,45 @@ class OpenRider(_Base):
             return True
         except Exception as e:
             print("!! Failed launching Rider:", str(e))
-            return False
\ No newline at end of file
+            return False
+
+class OpenXcode(_Base):
+    """ Opens a workspace in Xcode. """
+    platform = unrealcmd.Arg("Mac", "Platform to open (default = Mac)")
+
+    platform_to_workspace = {
+        "mac": "UE5 (Mac).xcworkspace",
+        "ios": "UE5 (IOS).xcworkspace",
+        "tvos":"UE5 (TVOS).xcworkspace",
+        "visionos": "UE5 (VisionOS).xcworkspace"
+    }
+
+    def main(self):
+        if sys.platform != "darwin":
+            self.print_error("Opening a solution in Xcode is only supported on Mac")
+            return False
+
+        self.print_info("Opening workspace in Xcode")
+
+        platform = self.args.platform.lower()
+        if platform not in self.platform_to_workspace:
+            self.print_error("Unsupported platform")
+            return False
+
+        workspace_name = self.platform_to_workspace[platform]
+        workspace_path = self.get_unreal_context().get_branch().get_dir() / Path(workspace_name)
+
+        print("Path:", workspace_path)
+        if not os.path.isdir(workspace_path):
+            self.print_error("Workspace dir not found")
+            return False
+
+        run_args = ("open", workspace_path)
+
+        try:
+            print("Launching Xcode:", *run_args)
+            proc = subprocess.Popen(run_args)
+            return True
+        except Exception as e:
+            print("!! Failed launching Xcode:", str(e))
+            return False
diff --git a/Engine/Extras/ushell/channels/unreal/core/describe.flow.py b/Engine/Extras/ushell/channels/unreal/core/describe.flow.py
--- a/Engine/Extras/ushell/channels/unreal/core/describe.flow.py
+++ b/Engine/Extras/ushell/channels/unreal/core/describe.flow.py
@@ -1,6 +1,7 @@
 # Copyright Epic Games, Inc. All Rights Reserved.
 
 import flow.describe
+import sys
 
 #-------------------------------------------------------------------------------
 fzf = flow.describe.Tool()
@@ -194,7 +195,10 @@ sln_generate.source("cmds/sln.py", "Generate")
 sln_generate.invoke("sln", "generate")
 
 sln_open = flow.describe.Command()
-sln_open.source("cmds/sln.py", "Open")
+if sys.platform == "darwin":
+    sln_open.source("cmds/sln.py", "OpenXcode")
+else:
+    sln_open.source("cmds/sln.py", "Open")
 sln_open.invoke("sln", "open")
 
 sln_10x = flow.describe.Command()
diff --git a/Engine/Extras/ushell/channels/unreal/core/pylib/unreal/_ini.py b/Engine/Extras/ushell/channels/unreal/core/pylib/unreal/_ini.py
--- a/Engine/Extras/ushell/channels/unreal/core/pylib/unreal/_ini.py
+++ b/Engine/Extras/ushell/channels/unreal/core/pylib/unreal/_ini.py
@@ -109,6 +109,7 @@ class Ini(object):
         self._lazy_re_compile()
         section = None
         for line in (x.strip() for x in file):
+            line = line.lstrip("\ufeff\ufffe")
             if m := Ini._comment_re.match(line):
                 continue
             elif m := Ini._section_re.match(line):
diff --git a/Engine/Extras/ushell/channels/unreal/core/pylib/unreal/_platform.py b/Engine/Extras/ushell/channels/unreal/core/pylib/unreal/_platform.py
--- a/Engine/Extras/ushell/channels/unreal/core/pylib/unreal/_platform.py
+++ b/Engine/Extras/ushell/channels/unreal/core/pylib/unreal/_platform.py
@@ -17,6 +17,8 @@ def _optional_api(func):
 
 #-------------------------------------------------------------------------------
 class Platform(object):
+    json_version_key = "MainVersion"
+
     @staticmethod
     def get_host():
         if sys.platform == "linux":  return "Linux"
@@ -74,7 +76,7 @@ class Platform(object):
     def get_unreal_context(self):
         return self._context
 
-    def _get_version_json(self, *, key="MainVersion", use_autosdk=True):
+    def _get_version_json(self, *, key=None, use_autosdk=True):
         def conform_json(inp):
             # As the time of writing the SDK.json files have trailing commas
             # and C-style comments. So we'll try and conform it
@@ -93,6 +95,9 @@ class Platform(object):
                 ret += sep + line
             return ret
 
+        # Nothing like a good bit of unstructuring.
+        key = key or self.json_version_key
+
         def read_sdk_json(path):
             if not path.is_file():
                 return
@@ -118,8 +123,16 @@ class Platform(object):
                 return version
 
             # ...or traverse up the platform forest
-            if next_path := sdk_info.get("ParentSDKFile", None):
-                return read_sdk_json(path.parent / next_path)
+            for upward_key in ("ParentSDKFile", "IncludeSDKFile"):
+                if not (next_path := sdk_info.get(upward_key, None)):
+                    continue
+
+                if next_path.startswith("$(EngineDir)"):
+                    engine_dir = self._context.get_engine().get_dir()
+                    next_path = engine_dir / next_path[13:]
+                else:
+                    next_path = path.parent / next_path
+                return read_sdk_json(next_path)
 
         name = self.get_config_name()
         engine_dir = self._context.get_engine().get_dir()
