Updated ushell [50202381]
#rnx

diff --git a/Engine/Extras/ushell/.version b/Engine/Extras/ushell/.version
--- a/Engine/Extras/ushell/.version
+++ b/Engine/Extras/ushell/.version
@@ -1 +1 @@
-49945479
\ No newline at end of file
+50202381
\ No newline at end of file
diff --git a/Engine/Extras/ushell/channels/unreal/core/pylib/unreal/_platform.py b/Engine/Extras/ushell/channels/unreal/core/pylib/unreal/_platform.py
--- a/Engine/Extras/ushell/channels/unreal/core/pylib/unreal/_platform.py
+++ b/Engine/Extras/ushell/channels/unreal/core/pylib/unreal/_platform.py
@@ -119,8 +119,9 @@ class Platform(object):
                 if version := sdk_info.get("AutoSDKDirectory", None):
                     return version
 
-            if version := sdk_info.get(key, None):
-                return version
+            for candidate in (key, "MainVersion"):
+                if version := sdk_info.get(candidate, None):
+                    return version
 
             # ...or traverse up the platform forest
             for upward_key in ("ParentSDKFile", "IncludeSDKFile"):
