Updated ushell [51334816]
#rnx

diff --git a/Engine/Extras/ushell/.version b/Engine/Extras/ushell/.version
--- a/Engine/Extras/ushell/.version
+++ b/Engine/Extras/ushell/.version
@@ -1 +1 @@
-50202381
\ No newline at end of file
+51334816
\ No newline at end of file
diff --git a/Engine/Extras/ushell/channels/flow/core/prompt.py b/Engine/Extras/ushell/channels/flow/core/prompt.py
--- a/Engine/Extras/ushell/channels/flow/core/prompt.py
+++ b/Engine/Extras/ushell/channels/flow/core/prompt.py
@@ -2,6 +2,7 @@
 
 import os
 import flow.cmd
+from pathlib import Path
 
 #-------------------------------------------------------------------------------
 class Prompt(flow.cmd.Cmd):
@@ -29,6 +30,15 @@ class Prompt(flow.cmd.Cmd):
             print(key, "=", value)
 
     def prompt(self, context):
-        cwd = os.getcwd()
+        try:
+            home = str(Path.home() if os.name != "nt" else "*")
+        except RuntimeError:
+            home = ""
+
+        cwd = Path(".").absolute()
+        cwd = str(cwd)
+        if cwd.startswith(home):
+            cwd = cwd.replace(home, "~");
+
         context.FLOW_CWD = cwd
         context.CWD = cwd
diff --git a/Engine/Extras/ushell/channels/flow/core/system/flow/_flick.py b/Engine/Extras/ushell/channels/flow/core/system/flow/_flick.py
--- a/Engine/Extras/ushell/channels/flow/core/system/flow/_flick.py
+++ b/Engine/Extras/ushell/channels/flow/core/system/flow/_flick.py
@@ -242,6 +242,8 @@ class Cmd(object):
             else:
                 # "Not" because only bool-type ops can have no parameter, and an
                 # _Opt(False, "desc") has a default of False but is True if parsed
+                if value is not None:
+                    raise SyntaxError(opt_name)
                 value = not opt.get_default()
 
             result[opt_name] = (opt.cast_value(value), False)
@@ -403,6 +405,9 @@ class Cmd(object):
         except ValueError as e:
             arg = " ".join(e.args[0])
             on_error(f"Unknown argument(s) '{arg}'")
+        except SyntaxError as e:
+            arg, = e.args
+            on_error(f"No value expectect for option '{e.args}'")
         except TypeError as e:
             argopt, value = e.args
             out_type = argopt.get_type().__name__
diff --git a/Engine/Extras/ushell/channels/flow/core/system/flow/describe.py b/Engine/Extras/ushell/channels/flow/core/system/flow/describe.py
--- a/Engine/Extras/ushell/channels/flow/core/system/flow/describe.py
+++ b/Engine/Extras/ushell/channels/flow/core/system/flow/describe.py
@@ -56,10 +56,12 @@ class Command(object):
 #-------------------------------------------------------------------------------
 class Channel(object):
     def __init__(self):
+        self._name = None
         self._parent = ""
         self._pips = []
         self._version = None
 
+    def name(self, name):       self._name = name
     def parent(self, parent):   self._parent = parent
     def version(self, version): self._version = version
 
diff --git a/Engine/Extras/ushell/channels/flow/core/system/flow/native/build.bat b/Engine/Extras/ushell/channels/flow/core/system/flow/native/build.bat
--- a/Engine/Extras/ushell/channels/flow/core/system/flow/native/build.bat
+++ b/Engine/Extras/ushell/channels/flow/core/system/flow/native/build.bat
@@ -22,7 +22,7 @@ if not "%3"=="release" (
     set _cl=/O1 /Os
 )
 
-set _cl="/I%pysdk%/include" /nologo /Zl /Zi /EHs-c- /c %_cl%
+set _cl="/I%pysdk%/include" /nologo /Zl /Zi /EHs-c- /c /FC %_cl%
 set _link=%_link% /nologo /dll /machine:x64^
     /nodefaultlib /align:16^
     /incremental:no /opt:ref,icf^
diff --git a/Engine/Extras/ushell/channels/flow/core/system/flow/native/native.cpp b/Engine/Extras/ushell/channels/flow/core/system/flow/native/native.cpp
--- a/Engine/Extras/ushell/channels/flow/core/system/flow/native/native.cpp
+++ b/Engine/Extras/ushell/channels/flow/core/system/flow/native/native.cpp
@@ -8,26 +8,26 @@
 //------------------------------------------------------------------------------
 static PyObject* isatty_impl(PyObject* self, PyObject** args, Py_ssize_t nargs)
 {
-    HANDLE Handle = nullptr;
+    HANDLE handle = nullptr;
 
     if (nargs > 0 && PyLong_Check(args[0]))
     {
         switch (PyLong_AsLong(args[0]))
         {
-        case 0: Handle = GetStdHandle(STD_INPUT_HANDLE);    break;
-        case 1: Handle = GetStdHandle(STD_OUTPUT_HANDLE);   break;
-        case 2: Handle = GetStdHandle(STD_ERROR_HANDLE);    break;
+        case 0: handle = GetStdHandle(STD_INPUT_HANDLE);    break;
+        case 1: handle = GetStdHandle(STD_OUTPUT_HANDLE);   break;
+        case 2: handle = GetStdHandle(STD_ERROR_HANDLE);    break;
         }
     }
 
-    if (Handle == nullptr)
+    if (handle == nullptr)
         Py_RETURN_NONE;
 
-    DWORD Type = GetFileType(Handle);
+    DWORD Type = GetFileType(handle);
     if (Type == FILE_TYPE_CHAR)
     {
         DWORD Mode;
-        if (GetConsoleMode(Handle, &Mode))
+        if (GetConsoleMode(handle, &Mode))
             Py_RETURN_TRUE;
     }
 
diff --git a/Engine/Extras/ushell/channels/flow/core/system/lib/bootstrap.py b/Engine/Extras/ushell/channels/flow/core/system/lib/bootstrap.py
--- a/Engine/Extras/ushell/channels/flow/core/system/lib/bootstrap.py
+++ b/Engine/Extras/ushell/channels/flow/core/system/lib/bootstrap.py
@@ -393,6 +393,11 @@ def _build_channel(channel_name:str, channel_dir:Path, cleaner:"_Cleaner") -> di
                 _log.print(" ...failed");
         _log.unindent()
 
+    # Channel override name
+    if channel._name:
+        channel_name = channel._name
+        _log.print("Name override:", channel_name)
+
     # Tools
     tools_manifest = {}
     tools = [x for x in read_items_by_type(flow.describe.Tool)]
@@ -639,11 +644,52 @@ class _Sources(object):
 
 
 
+#-------------------------------------------------------------------------------
+def _get_sources_inline(add_func) -> None:
+    for item in Path(__file__).parents:
+        if item.name == "channels":
+            add_func(item)
+            break
+
+#-------------------------------------------------------------------------------
+def _get_sources_platform(add_func) -> None:
+    # !! todo !!
+
+    # for item in Path(__file__).parents:
+    #     if item.name == "Engine":
+    #         break
+    # else:
+    #     return
+
+    # platforms_dir = item / "Platforms"
+    # if not platforms_dir.is_dir():
+    #     return
+
+    # for item in (x for x in platforms_dir.iterdir() if x.is_dir()):
+    #     candidate = item / "Extras/ushell/channel"
+    #     if candidate.is_dir():
+    #         name = item.name.lower()
+    #         add_func(candidate.parent, f"$platform.{name}.")
+    pass
+
+#-------------------------------------------------------------------------------
+def _get_sources_user(add_func) -> None:
+    from_home = Path.home() / ".ushell/channels"
+    if from_home.is_dir():
+        add_func(from_home, "$home.")
+
+    if from_env := os.getenv("FLOW_CHANNELS_DIR"):
+        import shlex
+        for item in shlex.split(from_env):
+            item = Path(item).resolve()
+            if item.is_dir():
+                add_func(item, "$env.")
+
 #-------------------------------------------------------------------------------
 def _get_sources() -> _Sources:
     sources = _Sources()
 
-    def _add_source(source_dir:Path, group_name:str="") -> int:
+    def add_source(source_dir:Path, group_name:str="") -> int:
         count = 0
         for item in source_dir.glob("*"):
             if not item.is_dir():
@@ -654,25 +700,13 @@ def _get_sources() -> _Sources:
                 sources.add_source(group_name + item.name, candidate)
                 count += 1
             elif not group_name:
-                count += _add_source(item, item.name + ".")
+                count += add_source(item, item.name + ".")
 
         return count
 
-    for item in Path(__file__).parents:
-        if item.name == "channels":
-            _add_source(item)
-            break
-
-    from_home = Path.home() / ".ushell/channels"
-    if from_home.is_dir():
-        _add_source(from_home, "$home.")
-
-    if from_env := os.getenv("FLOW_CHANNELS_DIR"):
-        import shlex
-        for item in shlex.split(from_env):
-            item = Path(item).resolve()
-            if item.is_dir():
-                _add_source(item, "$env.")
+    _get_sources_inline(add_source)
+    _get_sources_platform(add_source)
+    _get_sources_user(add_source)
 
     return sources
 
diff --git a/Engine/Extras/ushell/channels/flow/posix/shells.py b/Engine/Extras/ushell/channels/flow/posix/shells.py
--- a/Engine/Extras/ushell/channels/flow/posix/shells.py
+++ b/Engine/Extras/ushell/channels/flow/posix/shells.py
@@ -156,8 +156,8 @@ class _Zsh(_Base):
 def _get_complete_zsh():
     return r"""
 _flow_prompt() {
+    FLOW_CWD=$(pwd)
     eval $(\$prompt --format=sh)
-    FLOW_CWD=$(dirs)
 }
 setopt promptsubst
 if [[ $FLOW_PROMPT ]]; then
diff --git a/Engine/Extras/ushell/channels/unreal/core/cmds/clangdb.py b/Engine/Extras/ushell/channels/unreal/core/cmds/clangdb.py
--- a/Engine/Extras/ushell/channels/unreal/core/cmds/clangdb.py
+++ b/Engine/Extras/ushell/channels/unreal/core/cmds/clangdb.py
@@ -40,6 +40,7 @@ class ClangDb(unrealcmd.Cmd):
     keepgenfiles        = unrealcmd.Opt(False, "Add compile commands for UHT generated files (cache build time will increase)")
     keepplatformcode    = unrealcmd.Opt(False, "Add compile commands for Platforms files (cache build time will increase)")
     keepthirdpartycode  = unrealcmd.Opt(False, "Add compile commands for ThirdParty files (cache build time will increase)")
+    noallmodules = unrealcmd.Opt(False, "Do not add the -AllModules argument when running UnrealBuildTool")
 
     @unrealcmd.Cmd.summarise
     def _run_ubt(self, target, platform, ubt_args):
@@ -100,7 +101,6 @@ class ClangDb(unrealcmd.Cmd):
             excluded.append(".../ThirdParty/")
 
         ubt_args = (
-            "-allmodules",
             "-Mode=GenerateClangDatabase",
             f"-Exclude={";".join(excluded)}" if excluded else "",
             *self.args.ubtargs,
@@ -112,6 +112,9 @@ class ClangDb(unrealcmd.Cmd):
                 db_dir = project.get_dir()
                 ubt_args = ("-OutputDir=" + str(db_dir), *ubt_args)
 
+        if not self.args.noallmodules:
+            ubt_args = (*ubt_args, "-allmodules")
+
         if target.is_project_target() or self.args.projected:
             if project := ue_context.get_project():
                 ubt_args = ("-Project=" + str(project.get_path()), *ubt_args)
diff --git a/Engine/Extras/ushell/channels/unreal/core/cmds/getbuild.py b/Engine/Extras/ushell/channels/unreal/core/cmds/getbuild.py
--- a/Engine/Extras/ushell/channels/unreal/core/cmds/getbuild.py
+++ b/Engine/Extras/ushell/channels/unreal/core/cmds/getbuild.py
@@ -26,19 +26,20 @@ def run_cmd(cmd, dryrun = False):
 class GetBuild(unrealcmd.Cmd):
     """ Download a build/artifact from UE Cloud Storage
 
-Unless otherwise specified it is downloaded to <projectdir>/Staged/<platform>.
+    Unless otherwise specified it is downloaded to <projectdir>/Staged/<platform>.
 
-If your search only matches one build that build is downloaded for you, to add filtering you can use --match as follows:
-    * --match 8v93f39cvn3           - If you know the build id:
-    * --match 1337                  - To search for the build closest to a changelist
-    * --match job=v99f3;step=39f983 - For advanced searching you can specify key values
+    If your search only matches one build that build is downloaded for you, to
+    add filtering you can use --match as follows:
 
-If you find multiple builds that match you will be given the option to pick one, were the * indicates the build closest to your current commit."""
+        * --match 8v93f39cvn3        - A known build id.
+        * --match 1337               - Searches for the build closest to a changelist
+        * --match job=v99f;step=39f9 - Advanced searching by specify key values
 
-    buildtype           = unrealcmd.Arg(str, "The type of build to download. Usually packaged or staged builds.")
-    platforms           = unrealcmd.Arg([str], "The platform(s) you want to download from a build")
-    runargs             = unrealcmd.Arg([str], "Additional arguments to pass to the process being launched")
+    If you find multiple builds that match you will be given the option to pick
+    one, were the * indicates the build closest to your current commit."""
 
+    actions             = unrealcmd.Arg([str], "The platforms and build types to download data for")
+    runargs             = unrealcmd.Arg([str], "Additional arguments to pass to the process being launched")
     project             = unrealcmd.Opt("", "The project you want to download a build for. Defaults to the current project in ushell.")
     branch              = unrealcmd.Opt("", "Branch/Stream to fetch data for, defaults to current p4 stream. Accepts a p4 depot path like //UE5/Main or its converted form ue5-main")
     namespace           = unrealcmd.Opt("", "The namespace to use overriding what is defined in engine ini")
@@ -46,17 +47,36 @@ If you find multiple builds that match you will be given the option to pick one,
     clean               = unrealcmd.Opt(False, "Ignore any previously downloaded data and fetch it all again")
     verbose             = unrealcmd.Opt(False, "Enable verbose logging output")
     preflight           = unrealcmd.Opt(False, "If set, shows preflights")
-    wildcard            = unrealcmd.Opt("", "Windows style wildcard (using * and ?) to match file paths to include")
+    wildcard            = unrealcmd.Opt("", "Windows style wildcard (using * and ?) to match file paths to include (';' separated)")
     sync                = unrealcmd.Opt(False, "Use P4 to sync to the specified changelist.")
     buildgroup          = unrealcmd.Opt("", "The build group to filter, this is frequently globally unique and thus allows you to uniquely specify a build without any other options")
     config              = unrealcmd.Opt("", "The Configuration to filter")
-
+    nobins              = unrealcmd.Opt(False, "Add an exclusion for executable files (exe/dll/pdb/etc...)")
     dest                = unrealcmd.Opt("", "The directory to download the build to. Local contents will be overwritten. Defaults to `<projectdir>/Saved/<StagedBuilds/PackagedBuilds>/<platform>`")
     match               = unrealcmd.Opt("", "Either a build id (hexadecimal string), a changelist (number) or a key value separate list of key=values you want to find a build for. e.g. changelist=1337;job=foo")
 
+    _build_types = ("staged", "packaged")
+
+    def complete_actions(self, prefix):
+        for platform in unrealcmd.Cmd.complete_platform(self, prefix):
+            yield platform
+            if self.is_server_platform(platform):
+                yield platform + "-server"
+        yield from self._build_types
+        yield from self.complete_config
+
+    def complete_dest(self, prefix):
+        prefix = prefix or "."
+        for item in os.scandir(prefix):
+            if item.is_dir():
+                yield item.name + "/"
+
+    complete_config = ("development", "test", "shipping")
+
     def __init__(self):
         setattr(type(self), "exclude-wildcard", unrealcmd.Opt("", "Windows style wildcard (using * and ?) to match file paths to exclude. Applied after --wildcard include filter"))
         super().__init__()
+        self._query = {}
 
     def lookup_service_name(self):
         config = self.get_unreal_context().get_config()
@@ -114,25 +134,10 @@ If you find multiple builds that match you will be given the option to pick one,
 
         return str(bin_path)
 
-    def complete_buildtype(self, prefix):
-        return ["packaged", "staged"]
-
     def is_server_platform(self, platform):
         server_platforms = ("win64", "windows", "linux", "mac")
         return platform in server_platforms
 
-    def complete_platforms(self, prefix):
-        for platform in unrealcmd.Cmd.complete_platform(self, prefix):
-            yield platform
-            if self.is_server_platform(platform):
-                yield platform + "-server"
-
-    def complete_dest(self, prefix):
-        prefix = prefix or "."
-        for item in os.scandir(prefix):
-            if item.is_dir():
-                yield item.name + "/"
-
     def convert_platform(self, platform):
         # strip suffix to map platform name separately
         def strip_suffix(platform):
@@ -195,6 +200,15 @@ If you find multiple builds that match you will be given the option to pick one,
             return build_id_arg
         return None
 
+    def add_query_eq(self, key, value):
+        self._query[key] = {"$eq": value}
+
+    def add_query_neq(self, key, value):
+        self._query[key] = {"$neq": value}
+
+    def add_query_in(self, key, value):
+        self._query[key] = {"$in": value}
+
     def build_query(self, search_request, build_type, branch, platforms, preflight, buildgroup, config):
         search_changelist = None
         if search_request.isnumeric():
@@ -206,38 +220,32 @@ If you find multiple builds that match you will be given the option to pick one,
         if build_id:
             search_request = "" # special case handling the build_id filtering
 
-        query = {}
-
+        # Parse match string into query filters
         for statement in search_request.split(";"):
             if statement.find("!=") != -1:
-                k, v = statement.split("!=")
-                query[k] = {"$neq": v}
+                k, v = statement.split("!=", 1)
+                self.add_query_neq(k, v)
+            elif statement.find("~=") != -1:
+                k, v = statement.split("~=", 1)
+                self.add_query_in(k, v.split(","))
             elif statement.find("=") != -1:
-                k, v = statement.split("=")
-                query[k] = {"$eq": v}
-
-        if search_changelist:
-            query["commit"] = {"$eq": int(search_changelist)}
-
-        if build_id:
-            query["buildId"] = {"$eq": str(build_id)}
-
-        if build_type:
-            query["type"] = {"$eq": str(build_type)}
-
-        if buildgroup:
-            query["buildgroup"] = {"$eq": str(buildgroup)}
-
-        if config:
-            query["Configuration"] = {"$eq": str(config)}
-
-        # ignore the defaulted branch filtering when specifying build group as that uniquely identifies the branch anyway
-        # this allows us to download a build from any branch if you know its build group
+                k, v = statement.split("=", 1)
+                self.add_query_eq(k, v)
+
+        if search_changelist:   self.add_query_eq("commit", int(search_changelist))
+        if build_id:            self.add_query_eq("buildid", str(build_id))
+        if build_type:          self.add_query_eq("type", str(build_type))
+        if buildgroup:          self.add_query_eq("buildgroup", str(buildgroup))
+        if config:              self.add_query_eq("configuration", str(config))
+
+        # ignore the defaulted branch filtering when specifying build group as
+        # that uniquely identifies the branch anyway this allows us to download
+        # a build from any branch if you know its build group
         if branch and not buildgroup:
-            query["branch"] = {"$eq": str(branch.lower())}
+            self.add_query_eq("branch", str(branch.lower()))
 
         if not preflight:
-            query["ispreflight"] = {"$eq": "false"}
+            self.add_query_eq("ispreflight", "false")
 
         if platforms:
             platform_filters = []
@@ -252,9 +260,7 @@ If you find multiple builds that match you will be given the option to pick one,
                 if not p.endswith("-server") and self.is_server_platform(p):
                     platform_filters.append(f"{p}-server")
 
-            query["platform"] =  {"$in": platform_filters}
-
-        return query
+            self.add_query_in("platform", platform_filters)
 
     def search_build(self, host, override_host, namespace, request):
         # Run a search against the api to find builds that matches the request
@@ -288,7 +294,7 @@ If you find multiple builds that match you will be given the option to pick one,
             "buildgroup" : r["metadata"].get("buildgroup", ""),
             "commit": r["metadata"].get("commit", ""),
             "hordetemplate" : r["metadata"].get("hordetemplateid", ""),
-            "config": r["metadata"].get("Configuration", "")
+            "config": r["metadata"].get("Configuration", "All")
         } for r in results]
 
     def _show_select_build_prompt(self, branch_cl, builds):
@@ -373,13 +379,24 @@ If you find multiple builds that match you will be given the option to pick one,
 
         return result
 
+    def preprocess_actions(self):
+        self.args.buildtype = next((x for x in self.args.actions if x in self._build_types), "staged")
+        self.args.platforms = [x for x in self.args.actions if x not in self._build_types]
+        if not self.args.config:
+            self.args.config = next((x for x in self.args.actions if x in self.complete_config), "")
+
+    @unrealcmd.Cmd.summarise
     def main(self):
+        self.preprocess_actions()
+
         ue_context = self.get_unreal_context()
         project = self.args.project
         project_path = None
         if not project:
-            project = ue_context.get_project().get_name()
-            project_path = os.path.dirname(ue_context.get_project().get_path())
+            if not (project := ue_context.get_project()):
+                raise EnvironmentError("There is no active project set")
+            project_path = project.get_dir()
+            project = project.get_name()
 
         if len(self.args.platforms) == 0:
             self.print_error("You need to select at least one platform")
@@ -400,6 +417,18 @@ If you find multiple builds that match you will be given the option to pick one,
         else:
             namespace = self.lookup_namespace()
 
+        self.args.excludes = getattr(self.args, "exclude-wildcard", "")
+
+        for name in ("wildcard", "excludes"):
+            if "," in getattr(self.args, name):
+                parts = getattr(self.args, name).split(",")
+                setattr(self.args, name, ";".join(parts))
+
+        if self.args.nobins:
+            parts = self.args.excludes.split(";")
+            parts = (*parts, "*.exe", "*.dll", "*.self", "*.nss", "*.pdb", "*.map")
+            self.args.excludes = ";".join(parts)
+
         host = None
         override_host = None
         if self.args.host:
@@ -418,11 +447,11 @@ If you find multiple builds that match you will be given the option to pick one,
         builds_to_download = []
         found_builds = {}
 
-        query = self.build_query(self.args.match, build_type, branch, platforms_to_download, self.args.preflight, self.args.buildgroup, self.args.config)
+        self.build_query(self.args.match, build_type, branch, platforms_to_download, self.args.preflight, self.args.buildgroup, self.args.config)
         bucket_regex = f"{project}.{build_type}.*".lower()
         request = {
             "bucketRegex": bucket_regex,
-            "query": query
+            "query": self._query
         }
 
         builds = self.search_build(host, override_host, namespace, request)
@@ -490,7 +519,7 @@ If you find multiple builds that match you will be given the option to pick one,
                 "--local-path=" + build_dest,
                 "--build-id=" + build_id,
                 f"--wildcard={self.args.wildcard}" if self.args.wildcard else "",
-                "--exclude-wildcard=" + getattr(self.args, "exclude-wildcard", "")
+                "--exclude-wildcard=" + self.args.excludes
             )
 
             result |= self.run_zen(args)
diff --git a/Engine/Extras/ushell/channels/unreal/core/cmds/run.py b/Engine/Extras/ushell/channels/unreal/core/cmds/run.py
--- a/Engine/Extras/ushell/channels/unreal/core/cmds/run.py
+++ b/Engine/Extras/ushell/channels/unreal/core/cmds/run.py
@@ -258,6 +258,15 @@ class _Runtime(_RunCmd, _Attachable):
         cook_form = platform.get_cook_form(target_type.name.lower())
         stage_dir = project.get_dir() / f"Saved/StagedBuilds/{cook_form}"
 
+        # Detect the user doing -basedir= things themselves.
+        user_basedir = (x for x in self.args.runargs if x.startswith("-basedir="))
+        user_basedir = next(user_basedir, None)
+        if user_basedir:
+            if self.args.datadir:
+                self.print_warning("Ignoring --datadir as it conflicts with -basedir=")
+                self.args.datadir = ""
+            stage_dir = Path(user_basedir[9:]).absolute().parent.parent.parent
+
         # Allow the user to override the data directory used.
         if self.args.datadir:
             if self.args.cooked:
@@ -272,12 +281,11 @@ class _Runtime(_RunCmd, _Attachable):
             if not (stage_dir / "Engine").is_dir():
                 self.print_warning(f"An Engine/ directory is expected in '{stage_dir}'")
 
-
+        # Find the file in a manner that maintains case.
+        #
         # I probably shouldn't do this, but some platforms need this file to be
         # able to run with data served from the host PC. But I want to be explicit
         # and clear what command line arguments are given
-
-        # Find the file in a manner that maintains case.
         cmdline_txt = None
         maybe_names = (
             "UECommandLine.txt",
@@ -316,6 +324,11 @@ class _Runtime(_RunCmd, _Attachable):
             if platform.get_name() != unreal.Platform.get_host():
                 raise ValueError(f"The '--cooked' option may only be used when launching on {unreal.Platform.get_host()}")
 
+        # Stick with directories the user asked for.
+        if user_basedir:
+            print(stage_dir)
+            stage_dir = None
+
         # Validate data directories
         if stage_dir:
             if not stage_dir.is_dir():
@@ -324,7 +337,7 @@ class _Runtime(_RunCmd, _Attachable):
 
             if not stage_dir.is_dir():
                 raise EnvironmentError(f"No staged data found for '{cook_form}'")
-        else:
+        elif not user_basedir:
             cook_dir = project.get_dir() / f"Saved/Cooked/{cook_form}/"
             if not cook_dir.is_dir():
                 sandboxed = next((x for x in self.args.runargs if "-sandbox=" in x), None)
@@ -356,7 +369,7 @@ class _Runtime(_RunCmd, _Attachable):
         )
 
         if not self.args.cooked:
-            projectstore_marker = stage_dir / "ue.projectstore"
+            projectstore_marker = (stage_dir / "ue.projectstore") if stage_dir else Path()
             if not projectstore_marker.is_file():
                 args = (*args, "-pak")
 
diff --git a/Engine/Extras/ushell/channels/unreal/core/pylib/unreal/_context.py b/Engine/Extras/ushell/channels/unreal/core/pylib/unreal/_context.py
--- a/Engine/Extras/ushell/channels/unreal/core/pylib/unreal/_context.py
+++ b/Engine/Extras/ushell/channels/unreal/core/pylib/unreal/_context.py
@@ -217,6 +217,10 @@ class _Target(object):
     def get_type(self) -> TargetType:
         return self._type
 
+    def exists(self) -> bool:
+        path = self._ufs_mount.get_dir() / "Source" / (self.get_name() + ".Target.cs")
+        return path.is_file()
+
     def is_project_target(self) -> bool:
         return isinstance(self._ufs_mount.get_node(), _Project)
 
diff --git a/Engine/Extras/ushell/channels/unreal/core/pylib/unreal/platforms/_android.py b/Engine/Extras/ushell/channels/unreal/core/pylib/unreal/platforms/_android.py
--- a/Engine/Extras/ushell/channels/unreal/core/pylib/unreal/platforms/_android.py
+++ b/Engine/Extras/ushell/channels/unreal/core/pylib/unreal/platforms/_android.py
@@ -166,7 +166,17 @@ class PlatformBase(unreal.Platform):
             reverse_cmd = (str(adb), "reverse", "tcp:"+str(reverse_port_mapping), "tcp:"+str(reverse_port_mapping), ">", "NUL")
             os.system(" ".join(reverse_cmd))
 
-        args = "'" + " ".join(x.replace("\"", "\\\"") for x in args) + "'"
+        # Construct the arguments string
+        joined_args = " ".join(x.replace("\"", "\\\"") for x in args)
+
+        # Check the HOST platform (not the target Android platform)
+        if unreal.Platform.get_host() == "Mac":
+             # Mac requires escaping the single quote
+            args = "\\'" + joined_args + "\\'"
+        else:
+            # Windows/Linux use standard single quotes
+            args = "'" + joined_args + "'"
+
         cmd = (str(adb), "shell", "am", "start", "--es", "cmdline", args, intent)
         print(*cmd)
 
diff --git a/Engine/Extras/ushell/channels/unreal/core/pylib/vs/build.bat b/Engine/Extras/ushell/channels/unreal/core/pylib/vs/build.bat
--- a/Engine/Extras/ushell/channels/unreal/core/pylib/vs/build.bat
+++ b/Engine/Extras/ushell/channels/unreal/core/pylib/vs/build.bat
@@ -30,11 +30,11 @@ if not "%2"=="debug" (
     set _link=/debug
 )
 
-set _cl="/I%vssdk%." "/I%pysdk%/include" /nologo /Zl /Zi /EHs-c- /GS- /c %_cl%
+set _cl="/I%vssdk%." "/I%pysdk%/include" /nologo /Zl /Zi /EHs-c- /FC /GS- /c %_cl%
 set _link=%_link% /nologo /dll /machine:x64^
     /nodefaultlib /align:16^
     /incremental:no /opt:ref,icf^
-    "/libpath:%pysdk%/libs" python3.lib ole32.lib oleaut32.lib kernel32.lib user32.lib
+    "/libpath:%pysdk%/libs" python3.lib ole32.lib oleaut32.lib kernel32.lib user32.lib ntdll.lib
 
 mkdir 1>nul 2>nul _build
 cd _build
diff --git a/Engine/Extras/ushell/channels/unreal/core/pylib/vs/dte.cpp b/Engine/Extras/ushell/channels/unreal/core/pylib/vs/dte.cpp
--- a/Engine/Extras/ushell/channels/unreal/core/pylib/vs/dte.cpp
+++ b/Engine/Extras/ushell/channels/unreal/core/pylib/vs/dte.cpp
@@ -11,6 +11,19 @@
 #include <dte.h>
 #include <dte80.h>
 
+#include <winternl.h>
+
+//------------------------------------------------------------------------------
+struct Handle
+{
+            Handle(HANDLE handle)   { _handle = handle; }
+            ~Handle()               { if (_handle != nullptr || _handle != INVALID_HANDLE_VALUE) CloseHandle(_handle); }
+            operator HANDLE () const{ return _handle; }
+    HANDLE  _handle = nullptr;
+};
+
+
+
 //------------------------------------------------------------------------------
 template <typename T>
 class ComPtr
@@ -671,6 +684,45 @@ static PyObject* run_to_int3(PyObject* self, PyObject** args, Py_ssize_t nargs)
     Py_RETURN_TRUE;
 }
 
+//------------------------------------------------------------------------------
+static PyObject* get_proc_cwd(PyObject* self, PyObject** args, Py_ssize_t nargs)
+{
+    if (nargs != 1 || !PyLong_Check(args[0]))
+        Py_RETURN_NONE;
+
+    unsigned int pid = PyLong_AsLong(args[0]);
+    Handle proc = OpenProcess(PROCESS_QUERY_INFORMATION|PROCESS_VM_READ, FALSE, pid);
+
+    PROCESS_BASIC_INFORMATION pbi;
+    NtQueryInformationProcess(proc, ProcessBasicInformation, &pbi, sizeof(pbi), nullptr);
+
+    auto read_mem = [&] (auto addr, auto* out, int length=-1)
+    {
+        length = (length < 0) ? sizeof(*out) : length;
+        return (ReadProcessMemory(proc, (void*)addr, out, length, nullptr) != FALSE);
+    };
+
+    uintptr_t upp_address;
+    read_mem(uintptr_t(pbi.PebBaseAddress) + 0x20, &upp_address);
+
+    UNICODE_STRING cwd_str;
+    if (!read_mem(upp_address + 0x38, &cwd_str))
+        Py_RETURN_NONE;
+
+    enum { MaxSize = 512 };
+    char buf[MaxSize];
+    int length = cwd_str.Length;
+    if (length > MaxSize)
+        length = MaxSize;
+    if (!read_mem(cwd_str.Buffer, buf, length))
+        Py_RETURN_NONE;
+
+    auto* cwd = (wchar_t*)buf;
+    length = length / sizeof(wchar_t);
+    length -= (cwd[length - 1] == '\\');
+    return PyUnicode_FromWideChar(cwd, length);
+}
+
 
 
 //------------------------------------------------------------------------------
@@ -678,6 +730,7 @@ static PyMethodDef dte_module_methods[] = {
     { "suspend_process", PyCFunction(suspend_process), METH_FASTCALL, "Suspend a process by PID"},
     { "resume_process",  PyCFunction(resume_process),  METH_FASTCALL, "Resume a process by PID"},
     { "run_to_int3",     PyCFunction(run_to_int3),     METH_FASTCALL, "Run until NT's implicit int3s then suspends"},
+    { "get_proc_cwd",    PyCFunction(get_proc_cwd),    METH_FASTCALL, "Returns CWD of a process by ID"},
     {},
 };
 
