Updated ushell [51413913]
#rnx

diff --git a/Engine/Extras/ushell/.version b/Engine/Extras/ushell/.version
--- a/Engine/Extras/ushell/.version
+++ b/Engine/Extras/ushell/.version
@@ -1 +1 @@
-51334816
\ No newline at end of file
+51413913
\ No newline at end of file
diff --git a/Engine/Extras/ushell/channels/unreal/core/cmds/getbuild.py b/Engine/Extras/ushell/channels/unreal/core/cmds/getbuild.py
--- a/Engine/Extras/ushell/channels/unreal/core/cmds/getbuild.py
+++ b/Engine/Extras/ushell/channels/unreal/core/cmds/getbuild.py
@@ -17,7 +17,7 @@ def run_cmd(cmd, dryrun = False):
         print("== Running command: ", end="")
     print("." + cmd[1:])
     if not dryrun:
-        ret = subprocess.run(cmd)
+        ret = subprocess.run(cmd, shell=True)
         if ret.returncode:
             return (False, ret.returncode)
     return (True, "")
diff --git a/Engine/Extras/ushell/channels/unreal/perforce/cmds/clean.py b/Engine/Extras/ushell/channels/unreal/perforce/cmds/clean.py
--- a/Engine/Extras/ushell/channels/unreal/perforce/cmds/clean.py
+++ b/Engine/Extras/ushell/channels/unreal/perforce/cmds/clean.py
@@ -63,6 +63,7 @@ class _Collector(object):
         self._dirs = []
         self._careful_dirs = []
         self._pinned_files = set()
+        self._pinned_needles = []
 
     def get_pinned_count(self):
         return len(self._pinned_files)
@@ -74,6 +75,11 @@ class _Collector(object):
         path = str(path.resolve())
         self._pinned_files.add(path.lower())
 
+    def pin_needle(self, needle):
+        needle = needle.strip()
+        if needle:
+            self._pinned_needles.append(needle)
+
     def read_careful_dirs(self):
         return (x for x in self._careful_dirs)
 
@@ -87,7 +93,8 @@ class _Collector(object):
             for item in (x for x in dir.rglob("*") if x.is_file()):
                 key = str(item.resolve()).lower()
                 if key not in self._pinned_files:
-                    actuator.remove_file(item)
+                    if not any(needle in key for needle in self._pinned_needles):
+                        actuator.remove_file(item)
 
         actuator.end()
 
@@ -195,9 +202,10 @@ class Clean(flow.cmd.Cmd):
     Note that the removal happens in two stages. First directories are moved into
     .ushell_clean/ in the root of the branch. This directory is then removed in
     the background after `.p4 clean` exits."""
-    dryrun     = flow.cmd.Opt(False, "Do nothing except reports statistics")
+    dryrun     = flow.cmd.Opt(False, "Do nothing except reports statistics (this is slow...)")
     allsaved   = flow.cmd.Opt(False, "Completely clean Saved/ directories")
     savedkeeps = flow.cmd.Opt("Profiling,StagedBuilds", "Comma-separated list of Saved/ sub-directories to keep")
+    pinneedles = flow.cmd.Opt("site-packages", "If 'needle' is in deletee path then keep (comma-separated)")
 
     def _append_saved(self, collector, saved_dir):
         if self.args.allsaved:
@@ -276,9 +284,17 @@ class Clean(flow.cmd.Cmd):
         p4_have = P4.have(specs)
         for item in p4_have.read(on_error=False):
             path = Path(item.path)
+            # site-packages...
             collector.pin_file(path)
+
         print("done (", collector.get_pinned_count(), " files)", sep="")
 
+        print("Adding pin needles...")
+        for i, needle in enumerate(self.args.pinneedles.split(",")):
+            if needle:
+                print(f"{i:02}", needle)
+                collector.pin_needle(needle)
+
         if self.args.dryrun:
             collector.dispatch(_DryRun())
             return
