Updated ushell [51461529]
#rnx

diff --git a/Engine/Extras/ushell/.version b/Engine/Extras/ushell/.version
--- a/Engine/Extras/ushell/.version
+++ b/Engine/Extras/ushell/.version
@@ -1 +1 @@
-51413913
\ No newline at end of file
+51461529
\ No newline at end of file
diff --git a/Engine/Extras/ushell/channels/unreal/core/cmds/getbuild.py b/Engine/Extras/ushell/channels/unreal/core/cmds/getbuild.py
--- a/Engine/Extras/ushell/channels/unreal/core/cmds/getbuild.py
+++ b/Engine/Extras/ushell/channels/unreal/core/cmds/getbuild.py
@@ -345,7 +345,10 @@ class GetBuild(unrealcmd.Cmd):
                     cl = str(cl) + ("*" if cl == closest_cl else " ")
                 yield "%10s%s %s %s" % (cl, middle, name, config)
 
-        for line in fzf.run(list_available(closest_cl, available_builds), sort=False):
+        available = list(list_available(closest_cl, available_builds))
+        available.sort(reverse=True)
+
+        for line in fzf.run(available, sort=False):
             segments = line.split()
             build_name = segments[-2]
             config_name = segments[-1]
