Updated ushell [51492473]
#rnx

diff --git a/Engine/Extras/ushell/.version b/Engine/Extras/ushell/.version
--- a/Engine/Extras/ushell/.version
+++ b/Engine/Extras/ushell/.version
@@ -1 +1 @@
-51461529
\ No newline at end of file
+51492473
\ No newline at end of file
diff --git a/Engine/Extras/ushell/channels/unreal/core/cmds/getbuild.py b/Engine/Extras/ushell/channels/unreal/core/cmds/getbuild.py
--- a/Engine/Extras/ushell/channels/unreal/core/cmds/getbuild.py
+++ b/Engine/Extras/ushell/channels/unreal/core/cmds/getbuild.py
@@ -383,15 +383,27 @@ class GetBuild(unrealcmd.Cmd):
         return result
 
     def preprocess_actions(self):
+        non_platform_actions = set(self._build_types) | set(self.complete_config)
         self.args.buildtype = next((x for x in self.args.actions if x in self._build_types), "staged")
-        self.args.platforms = [x for x in self.args.actions if x not in self._build_types]
+        self.args.platforms = [x for x in self.args.actions if x not in non_platform_actions]
         if not self.args.config:
             self.args.config = next((x for x in self.args.actions if x in self.complete_config), "")
 
+    def _print_build_targets(self):
+        self.print_info("Dataset reqeusts")
+
+        print("Type:", self.args.buildtype)
+
+        for name in self.args.platforms:
+            print("Build:", name)
+
+        if match := self.args.match:
+            for match in match.split(";"):
+                if match := match.strip():
+                    print("Match:", match)
+
     @unrealcmd.Cmd.summarise
     def main(self):
-        self.preprocess_actions()
-
         ue_context = self.get_unreal_context()
         project = self.args.project
         project_path = None
@@ -401,6 +413,9 @@ class GetBuild(unrealcmd.Cmd):
             project_path = project.get_dir()
             project = project.get_name()
 
+        self.preprocess_actions()
+        self._print_build_targets()
+
         if len(self.args.platforms) == 0:
             self.print_error("You need to select at least one platform")
             return False
@@ -450,6 +465,7 @@ class GetBuild(unrealcmd.Cmd):
         builds_to_download = []
         found_builds = {}
 
+        self.print_info("Querying build list")
         self.build_query(self.args.match, build_type, branch, platforms_to_download, self.args.preflight, self.args.buildgroup, self.args.config)
         bucket_regex = f"{project}.{build_type}.*".lower()
         request = {
@@ -479,7 +495,10 @@ class GetBuild(unrealcmd.Cmd):
                 found_builds[build_key] = build
 
         if found_builds == None or len(found_builds) == 0:
-            self.print_error(f"Failed to find any build that matches query {request} in namespace {namespace}")
+            self.print_error(f"No builds found for namespace '{namespace}'")
+            print("Request = ", end="")
+            import pprint
+            pprint.pprint(request)
             return False
 
         selected_build = None
@@ -489,7 +508,7 @@ class GetBuild(unrealcmd.Cmd):
             print(f"Single build found, using build {selected_build_name}")
         else:
             # Ask the user to select a build
-            print("Select build to fetch:", end="")
+            self.print_info("Build selection")
             selected_build_id = self._show_select_build_prompt(branch_cl, found_builds)
             if not selected_build_id:
                 self.print_error("No build was selected")
