Fix: the perftest command wasn't forwarding the extra UAT args.

#rnx
#Synced-CL 51747525

diff --git a/Engine/Extras/ushell/channels/unreal/core/cmds/perftest.py b/Engine/Extras/ushell/channels/unreal/core/cmds/perftest.py
--- a/Engine/Extras/ushell/channels/unreal/core/cmds/perftest.py
+++ b/Engine/Extras/ushell/channels/unreal/core/cmds/perftest.py
@@ -231,7 +231,8 @@ class PerfTestBaseCmd(unrealcmd.Cmd):
             f"-resX={resx}",
             f"-resY={resy}",
             "-LocalReports",
-            *args
+            *args,
+            *self.args.uatargs
         )
 
         return self.run_perftest_uat_command(*perf_args)
