Updated ushell [51758756]
#rnx

diff --git a/Engine/Extras/ushell/.version b/Engine/Extras/ushell/.version
--- a/Engine/Extras/ushell/.version
+++ b/Engine/Extras/ushell/.version
@@ -1 +1 @@
-51492473
\ No newline at end of file
+51758756
\ No newline at end of file
diff --git a/Engine/Extras/ushell/channels/unreal/core/cmds/ai/build.md b/Engine/Extras/ushell/channels/unreal/core/cmds/ai/build.md
new file mode 100644
--- /dev/null
+++ b/Engine/Extras/ushell/channels/unreal/core/cmds/ai/build.md
@@ -0,0 +1,8 @@
+- How to build;
+	- Full build; ".build TARGET VARIANT"
+	- Individual module; ".build TARGET VARIANT <modulename>", where "modulename" can be derived by looking up the tree for "<modulename>.Build.cs"
+	  a file.
+	- Single .cpp file; ".build TARGET VARIANT <cpppath>"
+    - TARGET and VARIANT are mandatory (never empty)
+    - Never quote TARGET
+    - Defaults; TARGET=editor, VARIANT=development
diff --git a/Engine/Extras/ushell/channels/unreal/core/cmds/ai/cmd.py b/Engine/Extras/ushell/channels/unreal/core/cmds/ai/cmd.py
new file mode 100644
--- /dev/null
+++ b/Engine/Extras/ushell/channels/unreal/core/cmds/ai/cmd.py
@@ -0,0 +1,107 @@
+# Copyright Epic Games, Inc. All Rights Reserved.
+
+import os
+import tempfile
+import unrealcmd
+import subprocess as sp
+from pathlib import Path
+
+class Ai(unrealcmd.Cmd):
+    """
+    Launches an AI coding agent configured for ushell. Provides context
+    about building Unreal Engine files/modules, optimized file search
+    patterns using ripgrep, and modest source control.
+    """
+
+    def _add_tools(self, prompt_file):
+        print("Tools")
+
+        tools_path = Path(__file__).parent / "tools.md"
+        tools_content = tools_path.read_text()
+
+        # Search up the tree for code.rgignore and dirs.rgignore
+        rg_cmd = "rg"
+        for parent in Path(__file__).parents:
+            code_ignore = parent / "code.rgignore"
+            if code_ignore.exists():
+                rg_cmd += f" --ignore-file={code_ignore}"
+                dirs_ignore = parent / "dirs.rgignore"
+                if dirs_ignore.exists():
+                    rg_cmd += f" --ignore-file={dirs_ignore}"
+                break
+
+        tools_content = tools_content.replace("{{RG}}", rg_cmd)
+
+        prompt_file.write(tools_content)
+        prompt_file.write("\n")
+
+    def _add_ue_build(self, prompt_file):
+        print("Build")
+
+        build_path = Path(__file__).parent / "build.md"
+        build_content = build_path.read_text()
+        prompt_file.write(build_content)
+        prompt_file.write("\n")
+
+    def _add_ue_context(self, prompt_file, ue_context):
+        lines = []
+
+        if project := ue_context.get_project():
+            lines.append(f"- Project: {project.get_name()}")
+
+        if branch := ue_context.get_branch():
+            lines.append(f"- Branch: {branch.get_name()}")
+
+        engine_info = ue_context.get_engine().get_info()
+        if (cl := engine_info.get("Changelist", 0)) > 0:
+            lines.append(f"- Changelist: {cl}")
+
+        if lines:
+            print("Project/branch")
+            prompt_file.write("## Unreal Engine Context\n")
+            prompt_file.write("\n".join(lines))
+            prompt_file.write("\n\n")
+
+    def _add_ue(self, prompt_file):
+        try:
+            ue_context = self.get_unreal_context()
+        except EnvironmentError:
+            return
+
+        self._add_ue_build(prompt_file)
+        self._add_ue_context(prompt_file, ue_context)
+
+    def _add_perforce(self, prompt_file):
+        if Path(".git").is_dir():
+            return
+
+        print("Perforce")
+        perforce_path = Path(__file__).parent / "perforce.md"
+        perforce_content = perforce_path.read_text()
+        prompt_file.write(perforce_content)
+        prompt_file.write("\n")
+
+    def main(self):
+        try:
+            ue_context = self.get_unreal_context()
+            if branch := ue_context.get_branch():
+                os.chdir(branch.get_dir())
+        except EnvironmentError:
+            pass
+
+        self.print_info("Building prompt")
+        prompt_file = tempfile.NamedTemporaryFile(mode='w', delete=True, delete_on_close=False)
+
+        prompt_file.write("- You are a ushell user\n\n")
+        self._add_tools(prompt_file)
+        self._add_ue(prompt_file)
+        self._add_perforce(prompt_file)
+
+        prompt_file.close()
+
+        self.print_info("Launching claude")
+        cmd = (
+            "claude",
+            "--append-system-prompt-file", prompt_file.name
+        )
+        sp.run(cmd)
diff --git a/Engine/Extras/ushell/channels/unreal/core/cmds/ai/perforce.md b/Engine/Extras/ushell/channels/unreal/core/cmds/ai/perforce.md
new file mode 100644
--- /dev/null
+++ b/Engine/Extras/ushell/channels/unreal/core/cmds/ai/perforce.md
@@ -0,0 +1,3 @@
+- Source control if there is no known SCM marker;
+	- If a file is read-only, open it with "p4 edit <filepath>"
+	- Newly created files must be added with "p4 add"
diff --git a/Engine/Extras/ushell/channels/unreal/core/cmds/ai/tools.md b/Engine/Extras/ushell/channels/unreal/core/cmds/ai/tools.md
new file mode 100644
--- /dev/null
+++ b/Engine/Extras/ushell/channels/unreal/core/cmds/ai/tools.md
@@ -0,0 +1,6 @@
+- Tools;
+	- Primary tool for file globbing; "{{RG}} --files"
+	- Primary tool for file globbing by extension (EXT); "{{RG}} -g*.EXT --files"
+	- Priary tool for grep; "{{RG}}"
+
+- The repro is always too large to use the built-in Glob tool. Always use "rg" instead.
diff --git a/Engine/Extras/ushell/channels/unreal/core/cmds/build.py b/Engine/Extras/ushell/channels/unreal/core/cmds/build.py
--- a/Engine/Extras/ushell/channels/unreal/core/cmds/build.py
+++ b/Engine/Extras/ushell/channels/unreal/core/cmds/build.py
@@ -450,6 +450,7 @@ class Editor(_BuildCmd, unrealcmd.MultiPlatformCmd):
     """ Builds the editor. By default the current host platform is targeted. This
     can be overridden using the '--platform' option to cross-compile the editor."""
     variant     = _BuildCmd._variant
+    only        = unrealcmd.Opt(False, "Only the editor - infer all the --no* options")
     noscw       = unrealcmd.Opt(False, "Skip building ShaderCompileWorker")
     nopak       = unrealcmd.Opt(False, "Do not build UnrealPak")
     nointworker = unrealcmd.Opt(False, "Exclude the implicit InterchangeWorker dependency")
@@ -472,6 +473,7 @@ class Editor(_BuildCmd, unrealcmd.MultiPlatformCmd):
 
         editor_only = self.is_constrained_build() # i.e. single file or module
         editor_only |= bool(self.args.analyze)
+        editor_only |= self.args.only
         if not editor_only:
             targets_to_add = []
             if not self.args.noscw:
diff --git a/Engine/Extras/ushell/channels/unreal/core/cmds/build_xml.py b/Engine/Extras/ushell/channels/unreal/core/cmds/build_xml.py
--- a/Engine/Extras/ushell/channels/unreal/core/cmds/build_xml.py
+++ b/Engine/Extras/ushell/channels/unreal/core/cmds/build_xml.py
@@ -96,7 +96,9 @@ class _SetClearBase(unrealcmd.Cmd):
 
 #-------------------------------------------------------------------------------
 class _SetImpl(object):
-    """ Modify build configuration. """
+    """ Modify build configuration. Updates values in BuildConfiguration.xml files,
+    either globally or scoped to the current branch. 'section' and 'name' arguments
+    have tab-completion. """
     value = unrealcmd.Arg(str, "A value to set the property to")
 
     complete_value = _SetClearBase._complete_schema
diff --git a/Engine/Extras/ushell/channels/unreal/core/cmds/getbuild.py b/Engine/Extras/ushell/channels/unreal/core/cmds/getbuild.py
--- a/Engine/Extras/ushell/channels/unreal/core/cmds/getbuild.py
+++ b/Engine/Extras/ushell/channels/unreal/core/cmds/getbuild.py
@@ -36,7 +36,7 @@ class GetBuild(unrealcmd.Cmd):
         * --match job=v99f;step=39f9 - Advanced searching by specify key values
 
     If you find multiple builds that match you will be given the option to pick
-    one, were the * indicates the build closest to your current commit."""
+    one, where the * indicates the build closest to your current commit."""
 
     actions             = unrealcmd.Arg([str], "The platforms and build types to download data for")
     runargs             = unrealcmd.Arg([str], "Additional arguments to pass to the process being launched")
@@ -109,7 +109,7 @@ class GetBuild(unrealcmd.Cmd):
             if value := cloud.Host:
                 host=  str(value)
 
-        # if host is a semicolon seperated list we parse out the first entry as zen cli does not yet support a list of hosts
+        # if host is a semicolon separated list we parse out the first entry as zen cli does not yet support a list of hosts
         if host!= None and host.find(";") != -1:
             host = host.split(";")[0]
 
@@ -157,7 +157,7 @@ class GetBuild(unrealcmd.Cmd):
             converted_platform = self.get_platform(platform).get_config_name()
             return converted_platform + suffix
         except ValueError:
-            # unknown platform, just assume the name it valid
+            # unknown platform, just assume the name is valid
             return platform + suffix
 
     def convert_buildtype(self, build_type):
@@ -317,7 +317,7 @@ class GetBuild(unrealcmd.Cmd):
             x = x.replace("-server", "-s")
             x = x.replace("-game", "-g")
             vowels_plus = ("a", "e", "i", "o", "u", "y")
-            x = "".join(x for x in x if x not in vowels_plus)
+            x = "".join(c for c in x if c not in vowels_plus)
             return x[:7]
         columns = {x:abbreviate(x) for x in active_platforms}
 
@@ -363,20 +363,28 @@ class GetBuild(unrealcmd.Cmd):
                 raise Exception(f"Unable to map choice {line} back to selected build")
 
     def run_zen(self, args, include_passthrough_args = True):
-        ue_context = self.get_unreal_context()
-        engine = ue_context.get_engine()
+        override = os.environ.get("UE_ZEN_BINARY_PATH")
+        if override:
+            from pathlib import Path
+            bin_path = Path(override)
+        else:
+            ue_context = self.get_unreal_context()
+            engine = ue_context.get_engine()
 
-        bin_ext = ".exe" if os.name == "nt" else ""
-        bin_type = unreal.Platform.get_host()
+            bin_ext = ".exe" if os.name == "nt" else ""
+            bin_type = unreal.Platform.get_host()
 
-        bin_path = engine.get_dir()
-        bin_path /= f"Binaries/{bin_type}/zen{bin_ext}"
+            bin_path = engine.get_dir()
+            bin_path /= f"Binaries/{bin_type}/zen{bin_ext}"
 
         if not bin_path.is_file():
             raise FileNotFoundError(f"Unable to find '{bin_path}'")
 
+        global_args = os.environ.get("UE_ZEN_ARGS", "").split()
         if include_passthrough_args:
-            args = (*args, *self.args.runargs)
+            args = (*global_args, *args, *self.args.runargs)
+        else:
+            args = (*global_args, *args)
         cmd = self.get_exec_context().create_runnable(bin_path, *args)
         result = cmd.run()
 
@@ -390,7 +398,7 @@ class GetBuild(unrealcmd.Cmd):
             self.args.config = next((x for x in self.args.actions if x in self.complete_config), "")
 
     def _print_build_targets(self):
-        self.print_info("Dataset reqeusts")
+        self.print_info("Dataset requests")
 
         print("Type:", self.args.buildtype)
 
@@ -518,6 +526,14 @@ class GetBuild(unrealcmd.Cmd):
             selected_build_name = selected_build.get("buildgroup", selected_build["name"])
             print(f"Using build {selected_build_name}")
 
+        # If dest ends with '...' replace the suffix with the build's buildgroup
+        if dest.endswith("..."):
+            buildgroup = selected_build.get("buildgroup", "")
+            if not buildgroup:
+                self.print_error("Dest uses '...' suffix but the selected build has no buildgroup")
+                return False
+            dest = dest[:-3] + buildgroup
+
         # build the selection of build ids to download from which buckets
         for platform, (build_id, bucket) in selected_build["platforms"].items():
             builds_to_download.extend([(bucket, build_id, platform)])
diff --git a/Engine/Extras/ushell/channels/unreal/core/cmds/info.py b/Engine/Extras/ushell/channels/unreal/core/cmds/info.py
--- a/Engine/Extras/ushell/channels/unreal/core/cmds/info.py
+++ b/Engine/Extras/ushell/channels/unreal/core/cmds/info.py
@@ -166,7 +166,7 @@ class Projects(InfoBase):
 
 #-------------------------------------------------------------------------------
 class Config(InfoBase):
-    """ Evaluates currect context's config for a particular category. The 'override'
+    """ Evaluates current context's config for a particular category. The 'override'
     option can be used to enumerate platform-specific config sets."""
     category = unrealcmd.Arg(str, "The ini category to load")
     override = unrealcmd.Opt("", "Alternative config set to enumerate (e.g. 'WinGDK')")
diff --git a/Engine/Extras/ushell/channels/unreal/core/cmds/insights.py b/Engine/Extras/ushell/channels/unreal/core/cmds/insights.py
--- a/Engine/Extras/ushell/channels/unreal/core/cmds/insights.py
+++ b/Engine/Extras/ushell/channels/unreal/core/cmds/insights.py
@@ -5,7 +5,7 @@ from pathlib import Path
 
 #-------------------------------------------------------------------------------
 class Insights(unrealcmd.Cmd):
-    """ Conveninent access to UnrealInsights """
+    """ Convenient access to UnrealInsights """
     trace       = unrealcmd.Arg("", "Trace ident to open from server or .utrace file")
     uiargs      = unrealcmd.Arg([str], "Additional arguments to pass to UnrealInsights")
     build       = unrealcmd.Opt(False, "Build UnrealInsights before launching")
diff --git a/Engine/Extras/ushell/channels/unreal/core/cmds/notify.py b/Engine/Extras/ushell/channels/unreal/core/cmds/notify.py
--- a/Engine/Extras/ushell/channels/unreal/core/cmds/notify.py
+++ b/Engine/Extras/ushell/channels/unreal/core/cmds/notify.py
@@ -5,7 +5,7 @@ import unrealcmd
 import ctypes
 
 class Notify(unrealcmd.Cmd):
-    """ Flash the console window (if possible) to get the users attention """
+    """ Flash the console window (if possible) to get the user's attention """
 
     def main(self):
         if os.name == "nt":
diff --git a/Engine/Extras/ushell/channels/unreal/core/cmds/run.py b/Engine/Extras/ushell/channels/unreal/core/cmds/run.py
--- a/Engine/Extras/ushell/channels/unreal/core/cmds/run.py
+++ b/Engine/Extras/ushell/channels/unreal/core/cmds/run.py
@@ -211,7 +211,7 @@ class Program(_ProgramTarget):
 
 #-------------------------------------------------------------------------------
 class Target(_ProgramTarget):
-    """ Runs a name target on the host platform. """
+    """ Runs a named target on the host platform. """
     target  = unrealcmd.Arg(str, "Name of the target to run")
     variant = _RunCmd._variant
 
diff --git a/Engine/Extras/ushell/channels/unreal/core/describe.flow.py b/Engine/Extras/ushell/channels/unreal/core/describe.flow.py
--- a/Engine/Extras/ushell/channels/unreal/core/describe.flow.py
+++ b/Engine/Extras/ushell/channels/unreal/core/describe.flow.py
@@ -322,6 +322,11 @@ snapshot_list_cmd = flow.describe.Command()
 snapshot_list_cmd.source("cmds/snapshot.py", "List")
 snapshot_list_cmd.invoke("zen", "snapshot", "list")
 
+#-------------------------------------------------------------------------------
+ai = flow.describe.Command()
+ai.source("cmds/ai/cmd.py", "Ai")
+ai.invoke("ai")
+
 
 
 #-------------------------------------------------------------------------------
diff --git a/Engine/Extras/ushell/channels/unreal/core/pylib/unreal/_context.py b/Engine/Extras/ushell/channels/unreal/core/pylib/unreal/_context.py
--- a/Engine/Extras/ushell/channels/unreal/core/pylib/unreal/_context.py
+++ b/Engine/Extras/ushell/channels/unreal/core/pylib/unreal/_context.py
@@ -420,7 +420,7 @@ class _Project(_UfsNode):
     def get_path(self) -> Path:
         return self._path
 
-    def get_name(self) -> Path:
+    def get_name(self) -> str:
         return self._path.stem
 
     def get_engine(self) -> _Engine:
@@ -496,7 +496,7 @@ class _PlatformProvider(object):
 
     def __init__(self, temp_ctx:"Context"):
         self._temp_ctx = temp_ctx
-        self._provisions = None
+        self._provisions = {}
 
     @classmethod
     def _get_host_platform_name(self) -> str:
@@ -514,19 +514,17 @@ class _PlatformProvider(object):
         return self._host_name
 
     def _get_provisions(self) -> dict[str, Path]:
-        if self._provisions is not None:
+        if self._provisions:
             return self._provisions
 
-        provisions = self._provisions = {}
-
         # Built-in platforms
         host_name = self._get_host_platform_name()
         for item in Path(__file__).parent.glob(f"platforms/{host_name}/*.py"):
             if not item.name.startswith("_"):
-                provisions[item.stem] = item
+                self._provisions[item.stem] = item
 
         if host_name != "win64":
-            return provisions
+            return self._provisions
 
         # Use Engine/Platforms if we are hosted in a UE branch
         for candidate in Path(__file__).parent.parent.parents:
@@ -535,7 +533,7 @@ class _PlatformProvider(object):
                 if candidate.is_dir():
                     break
         else:
-            return provisions
+            return self._provisions
 
         for platform_dir in (x for x in candidate.glob("*") if x.is_dir()):
             extras_dir = platform_dir / "Extras/ushell"
@@ -543,9 +541,9 @@ class _PlatformProvider(object):
                 continue
 
             for item in extras_dir.glob("platform_*.py"):
-                provisions[item.stem[9:].lower()] = item
+                self._provisions[item.stem[9:].lower()] = item
 
-        return provisions
+        return self._provisions
 
     def get_platform(self, name:str|None=None) -> Platform|None:
         host_name = self._get_host_platform_name()
diff --git a/Engine/Extras/ushell/channels/unreal/core/pylib/unreal/platforms/_android.py b/Engine/Extras/ushell/channels/unreal/core/pylib/unreal/platforms/_android.py
--- a/Engine/Extras/ushell/channels/unreal/core/pylib/unreal/platforms/_android.py
+++ b/Engine/Extras/ushell/channels/unreal/core/pylib/unreal/platforms/_android.py
@@ -80,9 +80,10 @@ class PlatformBase(unreal.Platform):
             return
 
         for prefix in ("", unreal.Platform.get_sdks_dir()):
-            candidate = Path(prefix) / out
-            if candidate.is_dir():
-                return candidate
+            if prefix is not None:
+                candidate = Path(prefix) / out
+                if candidate.is_dir():
+                    return candidate
 
     def _get_adb(self):
         home_dir = self._get_android_home()
diff --git a/Engine/Extras/ushell/channels/unreal/core/pylib/zen/cmd.py b/Engine/Extras/ushell/channels/unreal/core/pylib/zen/cmd.py
--- a/Engine/Extras/ushell/channels/unreal/core/pylib/zen/cmd.py
+++ b/Engine/Extras/ushell/channels/unreal/core/pylib/zen/cmd.py
@@ -563,20 +563,27 @@ class ZenUtilityBaseCmd(ZenBaseCmd):
     zenargs  = unrealcmd.Arg([str], "Additional arguments to pass to zen utility")
 
     def get_zen_utility_command(self, cmdargs):
-        ue_context = self.get_unreal_context()
-        engine = ue_context.get_engine()
+        override = os.environ.get("UE_ZEN_BINARY_PATH")
+        if override:
+            from pathlib import Path
+            bin_path = Path(override)
+        else:
+            ue_context = self.get_unreal_context()
+            engine = ue_context.get_engine()
 
-        platform = unreal.Platform.get_host()
-        exe_filename = "zen"
-        if platform == "Win64":
-            exe_filename = exe_filename + ".exe"
+            platform = unreal.Platform.get_host()
+            exe_filename = "zen"
+            if platform == "Win64":
+                exe_filename = exe_filename + ".exe"
+
+            bin_path = engine.get_dir()
+            bin_path /= f"Binaries/{platform}/{exe_filename}"
 
-        bin_path = engine.get_dir()
-        bin_path /= f"Binaries/{platform}/{exe_filename}"
         if not bin_path.is_file():
             raise FileNotFoundError(f"Unable to find '{bin_path}'")
 
-        final_args = []
+        global_args = os.environ.get("UE_ZEN_ARGS", "").split()
+        final_args = list(global_args)
         if cmdargs: final_args += cmdargs
         if self.args.zenargs: final_args += self.args.zenargs
 
