Updated ushell [52208970]
#rnx

diff --git a/Engine/Extras/ushell/.version b/Engine/Extras/ushell/.version
--- a/Engine/Extras/ushell/.version
+++ b/Engine/Extras/ushell/.version
@@ -1 +1 @@
-51758756
\ No newline at end of file
+52208970
\ No newline at end of file
diff --git a/Engine/Extras/ushell/channels/flow/core/system/lib/bootstrap.py b/Engine/Extras/ushell/channels/flow/core/system/lib/bootstrap.py
--- a/Engine/Extras/ushell/channels/flow/core/system/lib/bootstrap.py
+++ b/Engine/Extras/ushell/channels/flow/core/system/lib/bootstrap.py
@@ -678,12 +678,20 @@ def _get_sources_user(add_func) -> None:
     if from_home.is_dir():
         add_func(from_home, "$home.")
 
-    if from_env := os.getenv("FLOW_CHANNELS_DIR"):
-        import shlex
-        for item in shlex.split(from_env):
-            item = Path(item).resolve()
-            if item.is_dir():
-                add_func(item, "$env.")
+    from_env = os.getenv("FLOW_CHANNELS_DIR")
+    if not from_env:
+        return
+
+    items = set()
+
+    import shlex
+    for item in shlex.split(from_env):
+        item = Path(item).resolve()
+        if item.is_dir():
+            items.add(item)
+
+    for item in items:
+        add_func(item, "$env.")
 
 #-------------------------------------------------------------------------------
 def _get_sources() -> _Sources:
diff --git a/Engine/Extras/ushell/channels/unreal/core/cmds/ai/cmd.py b/Engine/Extras/ushell/channels/unreal/core/cmds/ai/cmd.py
--- a/Engine/Extras/ushell/channels/unreal/core/cmds/ai/cmd.py
+++ b/Engine/Extras/ushell/channels/unreal/core/cmds/ai/cmd.py
@@ -104,4 +104,6 @@ class Ai(unrealcmd.Cmd):
             "claude",
             "--append-system-prompt-file", prompt_file.name
         )
-        sp.run(cmd)
+        env = os.environ.copy()
+        env["MSYS2_ARG_CONV_EXCL"] = "*"
+        sp.run(cmd, env=env)
diff --git a/Engine/Extras/ushell/channels/unreal/core/cmds/getbuild.py b/Engine/Extras/ushell/channels/unreal/core/cmds/getbuild.py
--- a/Engine/Extras/ushell/channels/unreal/core/cmds/getbuild.py
+++ b/Engine/Extras/ushell/channels/unreal/core/cmds/getbuild.py
@@ -367,6 +367,7 @@ class GetBuild(unrealcmd.Cmd):
         if override:
             from pathlib import Path
             bin_path = Path(override)
+            self.print_info(f"Using zen binary from UE_ZEN_BINARY_PATH: {bin_path}")
         else:
             ue_context = self.get_unreal_context()
             engine = ue_context.get_engine()
@@ -381,6 +382,8 @@ class GetBuild(unrealcmd.Cmd):
             raise FileNotFoundError(f"Unable to find '{bin_path}'")
 
         global_args = os.environ.get("UE_ZEN_ARGS", "").split()
+        if global_args:
+            self.print_info(f"Using extra zen args from UE_ZEN_ARGS: {' '.join(global_args)}")
         if include_passthrough_args:
             args = (*global_args, *args, *self.args.runargs)
         else:
diff --git a/Engine/Extras/ushell/channels/unreal/core/cmds/sln.py b/Engine/Extras/ushell/channels/unreal/core/cmds/sln.py
--- a/Engine/Extras/ushell/channels/unreal/core/cmds/sln.py
+++ b/Engine/Extras/ushell/channels/unreal/core/cmds/sln.py
@@ -52,7 +52,6 @@ class _Base(unrealcmd.MultiPlatformCmd):
         try:
             ppn_path = engine.get_dir() / "Intermediate/ProjectFiles/PrimaryProjectName.txt"
             if not ppn_path.is_file():
-                # Old path, needs to be maintained while projects prior to UE5.1 are supported
                 ppn_path = engine.get_dir() / ("Intermediate/ProjectFiles/M" + "asterProjectName.txt")
             with ppn_path.open("rb") as ppn_file:
                 name = ppn_file.read().decode().strip()
@@ -404,15 +403,16 @@ class File(_Base):
         if len(vs_instances) > 1:
             try:
                 sln_path = self._get_sln_path()
-                print("Using ", sln_path)
+                print("Using:", sln_path)
             except EnvironmentError:
                 sln_path = None
 
-            for instance in vs_instances:
-                candidate = instance.get_sln_path()
-                if candidate and sln_path.samefile(candidate):
-                    target_vs = instance
-                    break
+            if sln_path.is_file():
+                for instance in vs_instances:
+                    candidate = instance.get_sln_path()
+                    if candidate and sln_path.samefile(candidate):
+                        target_vs = instance
+                        break
         else:
             target_vs = vs_instances[0]
 
diff --git a/Engine/Extras/ushell/channels/unreal/core/pylib/vs/dte.cpp b/Engine/Extras/ushell/channels/unreal/core/pylib/vs/dte.cpp
--- a/Engine/Extras/ushell/channels/unreal/core/pylib/vs/dte.cpp
+++ b/Engine/Extras/ushell/channels/unreal/core/pylib/vs/dte.cpp
@@ -239,16 +239,40 @@ bool Dte::open_file(const wchar_t* path, int line) const
     if (FAILED(_dte->get_ItemOperations(item_ops)))
         return false;
 
-    const wchar_t* view_kind = L"{7651A703-06E5-11D1-8EBD-00A0C90F26EA}"; // from vsViewKindTextView
+    BStr path_str(path);
+    BStr view_kind(L"{7651A703-06E5-11D1-8EBD-00A0C90F26EA}"); // from vsViewKindTextView
+
+    VARIANT_BOOL is_open = FALSE;
+    item_ops->IsFileOpen(path_str, view_kind, &is_open);
+
+    ComPtr<Document> document;
+    _dte->get_ActiveDocument(document);
 
     ComPtr<Window> window;
-    if (FAILED(item_ops->OpenFile(BStr(path), BStr(view_kind), window)))
+    if (FAILED(item_ops->OpenFile(path_str, view_kind, window)))
         return false;
 
     if (line < 0)
         return true;
 
-    ComPtr<Document> document;
+    if (is_open == VARIANT_FALSE)
+    {
+        // File open is async. A crude attempt to wait for it without having to
+        // implement event subscription.
+
+        for (int i = 0; i++ < 16; ::Sleep(250))
+        {
+            ComPtr<Document> candidate;
+            if (FAILED(_dte->get_ActiveDocument(candidate)))
+                break;
+
+            if (candidate.get() != document.get())
+                break;
+        }
+
+        ::Sleep(250);
+    }
+
     if (SUCCEEDED(_dte->get_ActiveDocument(document)))
     {
         ComPtr<IDispatch> selection;
diff --git a/Engine/Extras/ushell/channels/unreal/core/pylib/zen/cmd.py b/Engine/Extras/ushell/channels/unreal/core/pylib/zen/cmd.py
--- a/Engine/Extras/ushell/channels/unreal/core/pylib/zen/cmd.py
+++ b/Engine/Extras/ushell/channels/unreal/core/pylib/zen/cmd.py
@@ -567,6 +567,7 @@ class ZenUtilityBaseCmd(ZenBaseCmd):
         if override:
             from pathlib import Path
             bin_path = Path(override)
+            self.print_info(f"Using zen binary from UE_ZEN_BINARY_PATH: {bin_path}")
         else:
             ue_context = self.get_unreal_context()
             engine = ue_context.get_engine()
@@ -583,6 +584,8 @@ class ZenUtilityBaseCmd(ZenBaseCmd):
             raise FileNotFoundError(f"Unable to find '{bin_path}'")
 
         global_args = os.environ.get("UE_ZEN_ARGS", "").split()
+        if global_args:
+            self.print_info(f"Using extra zen args from UE_ZEN_ARGS: {' '.join(global_args)}")
         final_args = list(global_args)
         if cmdargs: final_args += cmdargs
         if self.args.zenargs: final_args += self.args.zenargs
diff --git a/Engine/Extras/ushell/channels/unreal/perforce/cmds/sync.py b/Engine/Extras/ushell/channels/unreal/perforce/cmds/sync.py
--- a/Engine/Extras/ushell/channels/unreal/perforce/cmds/sync.py
+++ b/Engine/Extras/ushell/channels/unreal/perforce/cmds/sync.py
@@ -314,23 +314,49 @@ class Sync(_SyncBase):
             return False
         finally:
             if not self.args.dryrun:
-                # Record the synced changelist in Build.version
-                with open(build_ver_path, "r") as x:
-                    lines = list(x.readlines())
-
-                import stat
-                build_ver_mode = os.stat(build_ver_path).st_mode
-                os.chmod(build_ver_path, build_ver_mode|stat.S_IWRITE)
-
-                with open(build_ver_path, "w") as x:
-                    for line in lines:
-                        if r'"Changelist"' in line:
-                            line = line.split(":", 2)
-                            line = line[0] + f": {version_cl},\n"
-                        elif r'"BranchName"' in line:
-                            line = "\t\"BranchName\": \"X\"\n"
-                            line = line.replace("X", self._branch_root[:-1].replace("/", "+"))
-                        x.write(line)
+                self._record_cl(version_cl)
+
+    def _record_cl(self, version_cl):
+        self._update_build_version(version_cl)
+        self._update_ugs_state(version_cl)
+
+    def _update_build_version(self, version_cl):
+        build_ver_path = self._local_root + "Engine/Build/Build.version"
+        if not os.path.isfile(build_ver_path):
+            return
+
+        # Record the synced changelist in Build.version
+        with open(build_ver_path, "r") as x:
+            lines = list(x.readlines())
+
+        import stat
+        build_ver_mode = os.stat(build_ver_path).st_mode
+        os.chmod(build_ver_path, build_ver_mode|stat.S_IWRITE)
+
+        with open(build_ver_path, "w") as x:
+            for line in lines:
+                if r'"Changelist"' in line:
+                    line = line.split(":", 2)
+                    line = line[0] + f": {version_cl},\n"
+                elif r'"BranchName"' in line:
+                    line = "\t\"BranchName\": \"X\"\n"
+                    line = line.replace("X", self._branch_root[:-1].replace("/", "+"))
+                x.write(line)
+
+    def _update_ugs_state(self, version_cl):
+        state_path = Path(self._local_root) / ".ugs/state.json"
+        if not state_path.is_file():
+            return
+
+        import json
+        with open(state_path, "r") as x:
+            data = json.load(x)
+
+        data["LastSyncChangeNumber"] = version_cl
+
+        with open(state_path, "w") as x:
+            json.dump(data, x, indent="\t")
+            x.write("\n")
 
     def _apply_p4sync_txt(self, syncer):
         view_filter = syncer.get_view_filter()
