Updated ushell [52292358]
#rnx

diff --git a/Engine/Extras/ushell/.version b/Engine/Extras/ushell/.version
--- a/Engine/Extras/ushell/.version
+++ b/Engine/Extras/ushell/.version
@@ -1 +1 @@
-52208970
\ No newline at end of file
+52292358
\ No newline at end of file
diff --git a/Engine/Extras/ushell/channels/flow/nt/provision.bat b/Engine/Extras/ushell/channels/flow/nt/provision.bat
--- a/Engine/Extras/ushell/channels/flow/nt/provision.bat
+++ b/Engine/Extras/ushell/channels/flow/nt/provision.bat
@@ -7,9 +7,9 @@ goto:eof
 ::------------------------------------------------------------------------------
 :main
 
-set _pyver=3.13.3
-set _pysha=59ff76e16e6597de47474fb22be69e7191a89116910d728ab735079b078e52db
-set _pytag=313
+set _pyver=3.14.3
+set _pysha=ad4961a479dedbeb7c7d113253f8db1b1935586b73c27488712beec4f2c894e6
+set _pytag=314
 set _pymark=%_pyver%.version
 
 :: temp folder for detritus that's rmdir'd at the end
diff --git a/Engine/Extras/ushell/channels/unreal/core/cmds/cook.py b/Engine/Extras/ushell/channels/unreal/core/cmds/cook.py
--- a/Engine/Extras/ushell/channels/unreal/core/cmds/cook.py
+++ b/Engine/Extras/ushell/channels/unreal/core/cmds/cook.py
@@ -139,13 +139,33 @@ class Cook(_Cook):
 
 #-------------------------------------------------------------------------------
 class _Runtime(_Cook):
+    """ If the given platform is known the environment for the platform will be
+    derived and details printed. Otherwise the runtime target will be used for a
+    best-effort guess. """
+
     platform = unrealcmd.Arg("", "The platform to cook data for")
 
+    def _main_unplatformed(self, target):
+        cook_form = self.args.platform
+        if target == unreal.TargetType.CLIENT:
+            cook_form += "Client"
+
+        self.print_info("Platform")
+        print("Given:", self.args.platform)
+        print("Form:", cook_form)
+
+        exec_context = super().get_exec_context()
+        return self._cook(cook_form, exec_context)
+
     def main(self, target):
         platform = self.args.platform
         if not platform:
             platform = unreal.Platform.get_host()
-        platform = self.get_platform(platform)
+
+        try:
+            platform = self.get_platform(platform)
+        except ValueError:
+            return self._main_unplatformed(target)
 
         exec_context = super().get_exec_context()
         exec_env = exec_context.get_env()
