Updated ushell [52344128]
#rnx

diff --git a/Engine/Extras/ushell/.version b/Engine/Extras/ushell/.version
--- a/Engine/Extras/ushell/.version
+++ b/Engine/Extras/ushell/.version
@@ -1 +1 @@
-52292358
\ No newline at end of file
+52344128
\ No newline at end of file
diff --git a/Engine/Extras/ushell/channels/flow/core/boot.py b/Engine/Extras/ushell/channels/flow/core/boot.py
--- a/Engine/Extras/ushell/channels/flow/core/boot.py
+++ b/Engine/Extras/ushell/channels/flow/core/boot.py
@@ -14,14 +14,15 @@ class Boot(flow.cmd.Cmd):
 
     def print_title(self):
         lines = (
-            "         _       _ _ ",
+            "        ._.      _ _ ",
             " _ _ ___| |_ ___| | |",
             "| | |_ -|   | -_| | |",
-            "|___|___|_|_|___|_|_|",
+            "|___|___|_|_|___|_|_| ",
         )
-        print("\x1b[93m", end="")
         for i, line in enumerate(lines):
-            print("  ", line, sep="")
+            # print(f"  \x1b[9{i + 2};4m", end="")
+            u = ";4" if i == 3 else ""
+            print(" ", f"\x1b[9{i + 1}{u}m ", line, "\x1b[0m", sep="")
         print("\n\x1b[0m", end="")
 
     def main(self):
diff --git a/Engine/Extras/ushell/channels/flow/core/system/lib/bootstrap.py b/Engine/Extras/ushell/channels/flow/core/system/lib/bootstrap.py
--- a/Engine/Extras/ushell/channels/flow/core/system/lib/bootstrap.py
+++ b/Engine/Extras/ushell/channels/flow/core/system/lib/bootstrap.py
@@ -637,6 +637,7 @@ class _Sources(object):
     def get_age(self) -> float:
         deps = (
             Path(sys.executable),
+            Path(sys.executable).parent,
             Path(__file__).parent.parent / "version",
             *(self._sources.values()),
         )
diff --git a/Engine/Extras/ushell/channels/flow/core/system/version b/Engine/Extras/ushell/channels/flow/core/system/version
--- a/Engine/Extras/ushell/channels/flow/core/system/version
+++ b/Engine/Extras/ushell/channels/flow/core/system/version
@@ -1 +1 @@
-25
+26
diff --git a/Engine/Extras/ushell/channels/unreal/core/pylib/zen/cmd.py b/Engine/Extras/ushell/channels/unreal/core/pylib/zen/cmd.py
--- a/Engine/Extras/ushell/channels/unreal/core/pylib/zen/cmd.py
+++ b/Engine/Extras/ushell/channels/unreal/core/pylib/zen/cmd.py
@@ -95,6 +95,94 @@ def _get_ips_nt(target_platform):
 
 _get_ips = _get_ips_nt if os.name == "nt" else (_get_ips_osx if sys.platform == "darwin" else _get_ips_socket)
 
+#-------------------------------------------------------------------------------
+def _unescape_ini_quoted_string(value):
+    """Unescape a UE INI quoted string value, matching FParse::QuotedString."""
+    result = []
+    i = 0
+    while i < len(value):
+        ch = value[i]
+        if ch != '\\':
+            result.append(ch)
+            i += 1
+            continue
+        i += 1
+        if i >= len(value):
+            result.append('\\')
+            break
+        esc = value[i]
+        if   esc == '\\': result.append('\\')
+        elif esc == '"':  result.append('"')
+        elif esc == "'":  result.append("'")
+        elif esc == 'n':  result.append('\n')
+        elif esc == 'r':  result.append('\r')
+        elif esc == 't':  result.append('\t')
+        else:
+            result.append('\\')
+            result.append(esc)
+        i += 1
+    return "".join(result)
+
+def _read_ini_value(ini_path, section, key):
+    """Read a value from a UE-style INI file given a section and key name.
+
+    Handles UE's quoted-string format: if the value is wrapped in double quotes,
+    the quotes are stripped and escape sequences are processed to match the
+    behavior of FParse::QuotedString / FConfigFile::ShouldExportQuotedString."""
+    in_section = False
+    section_header = f"[{section}]"
+    with open(ini_path, "r") as f:
+        for line in f:
+            line = line.strip()
+            if line.startswith("["):
+                in_section = line.lower() == section_header.lower()
+                continue
+            if in_section:
+                eq_index = line.find("=")
+                if eq_index > 0 and line[:eq_index].lower() == key.lower():
+                    raw = line[eq_index + 1:]
+                    if raw.startswith('"') and raw.endswith('"'):
+                        return _unescape_ini_quoted_string(raw[1:-1])
+                    return raw
+    return None
+
+def _get_zen_host_auth():
+    """Retrieve the Zen host auth data from the platform credential store.
+
+    Matches the storage locations used by FPlatformMisc::GetStoredValue with
+    store ID "Epic Games", section "Zen", key "HostAuth":
+      - Windows: HKCU\\Software\\Epic Games\\Zen\\HostAuth (registry)
+      - macOS:   ~/Library/Application Support/Epic/Epic Games/KeyValueStore.ini [Zen] HostAuth
+      - Linux:   ~/.config/Epic/Epic Games/KeyValueStore.ini [Zen] HostAuth
+    Returns the parsed auth dict or None if not found.
+    """
+    try:
+        auth_json = None
+        if os.name == "nt":
+            import winreg
+            with winreg.OpenKey(winreg.HKEY_CURRENT_USER, r"Software\Epic Games\Zen") as key:
+                auth_json, _ = winreg.QueryValueEx(key, "HostAuth")
+        else:
+            # On Mac and Linux, FPlatformMisc::SetStoredValue writes to an INI file:
+            #   ApplicationSettingsDir() / "Epic Games" / "KeyValueStore.ini"
+            # with section [Zen] and key HostAuth=<value>
+            if sys.platform == "darwin":
+                app_settings_dir = os.path.join(os.path.expanduser("~"),
+                    "Library", "Application Support", "Epic")
+            else:
+                app_settings_dir = os.path.join(os.path.expanduser("~"),
+                    ".config", "Epic")
+
+            ini_path = os.path.join(app_settings_dir, "Epic Games", "KeyValueStore.ini")
+            if os.path.isfile(ini_path):
+                auth_json = _read_ini_value(ini_path, "Zen", "HostAuth")
+
+        if auth_json:
+            return json.loads(auth_json)
+    except Exception:
+        pass
+    return None
+
 #-------------------------------------------------------------------------------
 class ZenBaseCmd(unrealcmd.Cmd):
     cloudauthservice = unrealcmd.Opt("", "Name of the service to authorize with when importing from a cloud source")
@@ -354,15 +442,19 @@ class ZenBaseCmd(unrealcmd.Cmd):
             os.makedirs(target_gcmarker_dir)
         remotehostnames = _get_ips(snapshot['targetplatform'])
         remotehostnames.append("hostname://" + socket.getfqdn())
-        projectstore = {
-            "zenserver" : {
+        zenserver_data = {
                 "islocalhost" : True,
                 "hostname" : "[::1]",
                 "remotehostnames" : remotehostnames,
                 "hostport" : 8558,
-                "projectid" : target_projectid,
-                "oplogid" : target_oplogid
-            }
+        }
+        host_auth = _get_zen_host_auth()
+        if host_auth:
+            zenserver_data["hostauth"] = host_auth
+        zenserver_data["projectid"] = target_projectid
+        zenserver_data["oplogid"] = target_oplogid
+        projectstore = {
+            "zenserver" : zenserver_data
         }
         with open(target_gcmarker, "w") as target_gcmarker_outfile:
             json.dump(projectstore, target_gcmarker_outfile, indent='\t')
