Updated ushell [52516252]
#rnx

diff --git a/Engine/Extras/ushell/.version b/Engine/Extras/ushell/.version
--- a/Engine/Extras/ushell/.version
+++ b/Engine/Extras/ushell/.version
@@ -1 +1 @@
-52344128
\ No newline at end of file
+52516252
\ No newline at end of file
diff --git a/Engine/Extras/ushell/channels/unreal/core/pylib/zen/cmd.py b/Engine/Extras/ushell/channels/unreal/core/pylib/zen/cmd.py
--- a/Engine/Extras/ushell/channels/unreal/core/pylib/zen/cmd.py
+++ b/Engine/Extras/ushell/channels/unreal/core/pylib/zen/cmd.py
@@ -516,7 +516,6 @@ class ZenBaseCmd(unrealcmd.Cmd):
         "oplog-import",
         self.get_project_or_default(projectid),
         oplogid or snapshot['targetplatform'],
-        "--ignore-missing-attachments",
         "--clean"
         ]
 
